[SerializationUtils] Moved plugin

diff --git a/Engine/Plugins/Experimental/SerializationUtils/SerializationUtils.uplugin b/Engine/Plugins/Experimental/SerializationUtils/SerializationUtils.uplugin
new file mode 100644
--- /dev/null
+++ b/Engine/Plugins/Experimental/SerializationUtils/SerializationUtils.uplugin
@@ -0,0 +1,29 @@
+{
+	"FileVersion": 3,
+	"Version": 0.1,
+	"VersionName": "0.1.0-alpha",
+	"FriendlyName": "SerializationUtils",
+	"Description": "Utilities for serialization (xml, json, etc) with extended functionality.",
+	"Category": "Experimental",
+	"CreatedBy": "Epic Games, Inc.",
+	"CreatedByURL": "http://epicgames.com",
+	"DocsURL": "",
+	"MarketplaceURL": "",
+	"SupportURL": "",
+	"CanContainContent": true,
+	"IsBetaVersion": false,
+	"IsExperimentalVersion": false,
+	"Installed": false,
+	"Modules": [
+		{
+			"Name": "JsonSerialization",
+			"Type": "Runtime",
+			"LoadingPhase": "Default"
+		},
+		{
+			"Name": "XmlSerialization",
+			"Type": "Runtime",
+			"LoadingPhase": "Default"
+		}
+	]
+}
\ No newline at end of file
diff --git a/Engine/Plugins/Experimental/SerializationUtils/Source/JsonSerialization/JsonSerialization.Build.cs b/Engine/Plugins/Experimental/SerializationUtils/Source/JsonSerialization/JsonSerialization.Build.cs
new file mode 100644
--- /dev/null
+++ b/Engine/Plugins/Experimental/SerializationUtils/Source/JsonSerialization/JsonSerialization.Build.cs
@@ -0,0 +1,28 @@
+// Copyright Epic Games, Inc. All Rights Reserved.
+
+using System.IO;
+using UnrealBuildTool;
+
+public class JsonSerialization : ModuleRules
+{
+	public JsonSerialization(ReadOnlyTargetRules Target) : base(Target)
+	{
+		PCHUsage = ModuleRules.PCHUsageMode.UseExplicitOrSharedPCHs;
+
+		PublicDependencyModuleNames.AddRange(
+			new string[]
+			{
+				"Core",
+				"CoreUObject",
+				"Json"
+			}
+		);
+
+		PrivateDependencyModuleNames.AddRange(
+			new string[]
+			{
+				"Serialization",
+			}
+		);
+	}
+}
diff --git a/Engine/Plugins/Experimental/SerializationUtils/Source/JsonSerialization/Private/Formatters/JsonArchiveInputFormatterEx.cpp b/Engine/Plugins/Experimental/SerializationUtils/Source/JsonSerialization/Private/Formatters/JsonArchiveInputFormatterEx.cpp
new file mode 100644
--- /dev/null
+++ b/Engine/Plugins/Experimental/SerializationUtils/Source/JsonSerialization/Private/Formatters/JsonArchiveInputFormatterEx.cpp
@@ -0,0 +1,718 @@
+// Copyright Epic Games, Inc. All Rights Reserved.
+
+#include "Formatters/JsonArchiveInputFormatterEx.h"
+#include "Dom/JsonObject.h"
+#include "Misc/Base64.h"
+#include "Misc/SecureHash.h"
+#include "Serialization/JsonReader.h"
+#include "Serialization/JsonSerializer.h"
+#include "UObject/LazyObjectPtr.h"
+#include "UObject/ObjectPtr.h"
+#include "UObject/SoftObjectPath.h"
+#include "UObject/SoftObjectPtr.h"
+#include "UObject/UnrealType.h"
+#include "UObject/WeakObjectPtr.h"
+
+#if WITH_TEXT_ARCHIVE_SUPPORT
+
+namespace UE::JsonArchiveInputFormatter::Private
+{
+	void SerializeObject(FStructuredArchiveRecord& InRootRecord
+	, UObject* InObject
+	, const TFunction<bool(FProperty*)>& ShouldSkipProperty)
+	{
+		if (!InObject)
+		{
+			return;
+		}
+	
+		FStructuredArchiveRecord ObjectRecord = InRootRecord.EnterRecord(*InObject->GetName());
+		FArchive& UnderlyingArchive = InRootRecord.GetUnderlyingArchive();
+	
+		for (FProperty* Property = InObject->GetClass()->PropertyLink; Property; Property = Property->PropertyLinkNext)
+		{
+			if (!Property || Property->HasAnyPropertyFlags(CPF_Transient) || (ShouldSkipProperty && ShouldSkipProperty(Property)))
+			{
+				continue;
+			}
+
+			FStructuredArchiveSlot PropertyField = InRootRecord.EnterField(*Property->GetName());
+
+			for (int32 Index = 0; Index < Property->ArrayDim; ++Index)
+			{
+				UObject* const Archetype = InObject->GetArchetype();
+			
+				void* const Data = Property->ContainerPtrToValuePtr<void>(InObject, Index);
+				const void* const Defaults = Property->ContainerPtrToValuePtrForDefaults<void>(Archetype->GetClass()
+					, Archetype
+					, Index);
+			
+				{
+					FSerializedPropertyScope SerializedProperty(UnderlyingArchive, Property);
+					Property->SerializeItem(PropertyField, Data, Defaults);
+				}
+			}
+		}
+	}
+};
+
+FJsonArchiveInputFormatterEx::FJsonArchiveInputFormatterEx(FArchive& InInner, UObject* InRootObject, TFunction<UObject* (const FPackageIndex)> InResolveObject) 
+	: Inner(InInner)
+	, ResolveObject(InResolveObject)
+	, RootObject(InRootObject)
+{
+	Inner.SetIsTextFormat(true);
+	Inner.ArAllowLazyLoading = false;
+
+	TSharedPtr< FJsonObject > RootJsonObject;
+	TSharedRef< TJsonReader<UTF8CHAR> > Reader = TJsonReaderFactory<UTF8CHAR>::Create(&InInner);
+	ensure(FJsonSerializer::Deserialize(Reader, RootJsonObject, FJsonSerializer::EFlags::None));
+	ValueStack.Add(MakeShared<FJsonValueObject>(RootJsonObject));
+
+	ValueStack.Reserve(64);
+	ArrayValuesRemainingStack.Reserve(64);
+}
+
+FJsonArchiveInputFormatterEx::~FJsonArchiveInputFormatterEx()
+{
+}
+
+FArchive& FJsonArchiveInputFormatterEx::GetUnderlyingArchive()
+{
+	return Inner;
+}
+
+FStructuredArchiveFormatter* FJsonArchiveInputFormatterEx::CreateSubtreeReader()
+{
+	FJsonArchiveInputFormatterEx* Cloned = new FJsonArchiveInputFormatterEx(*this);
+	Cloned->ObjectStack.Empty();
+	Cloned->ValueStack.Empty();
+	Cloned->MapIteratorStack.Empty();
+	Cloned->ArrayValuesRemainingStack.Empty();
+	Cloned->ValueStack.Push(ValueStack.Top());
+
+	return Cloned;
+}
+
+bool FJsonArchiveInputFormatterEx::HasDocumentTree() const
+{
+	return true;
+}
+
+void FJsonArchiveInputFormatterEx::EnterRecord()
+{
+	TSharedPtr<FJsonValue>& Value = ValueStack.Top();
+	ObjectStack.Emplace(Value->AsObject(), ValueStack.Num());
+}
+
+void FJsonArchiveInputFormatterEx::LeaveRecord()
+{
+	// In the original, this check is in place for debugging purposes and to check that all values are consumed
+	// check(ValueStack.Num() == ObjectStack.Top().ValueCountOnCreation);
+	ObjectStack.Pop();
+}
+
+void FJsonArchiveInputFormatterEx::EnterField(FArchiveFieldName Name)
+{
+	TSharedPtr<FJsonObject>& NewObject = ObjectStack.Top().JsonObject;
+	TSharedPtr<FJsonValue> Field = NewObject->TryGetField(EscapeFieldName(Name.Name));
+	check(Field.IsValid());
+	ValueStack.Add(Field);
+}
+
+void FJsonArchiveInputFormatterEx::LeaveField()
+{
+	ValueStack.Pop();
+}
+
+bool FJsonArchiveInputFormatterEx::TryEnterField(FArchiveFieldName Name, bool bEnterWhenSaving)
+{
+	TSharedPtr<FJsonObject>& NewObject = ObjectStack.Top().JsonObject;
+	TSharedPtr<FJsonValue> Field = NewObject->TryGetField(EscapeFieldName(Name.Name));
+	if (Field.IsValid())
+	{
+		ValueStack.Add(Field);
+		return true;
+	}
+	return false;
+}
+
+void FJsonArchiveInputFormatterEx::EnterArray(int32& NumElements)
+{
+	TSharedPtr<FJsonValue>& Value = ValueStack.Top();
+
+	const TArray<TSharedPtr<FJsonValue>>& ArrayValue = Value->AsArray();
+	for(int Idx = ArrayValue.Num() - 1; Idx >= 0; Idx--)
+	{
+		ValueStack.Add(ArrayValue[Idx]);
+	}
+
+	NumElements = ArrayValue.Num();
+	ArrayValuesRemainingStack.Add(NumElements);
+}
+
+void FJsonArchiveInputFormatterEx::LeaveArray()
+{
+	check(ArrayValuesRemainingStack.Num() > 0);
+	int32 Remaining = ArrayValuesRemainingStack.Top();
+	ArrayValuesRemainingStack.Pop();
+	check(Remaining >= 0);
+	check(Remaining <= ValueStack.Num());
+	while (Remaining-- > 0)
+	{
+		ValueStack.Pop();
+	}
+}
+
+void FJsonArchiveInputFormatterEx::EnterArrayElement()
+{
+	check(ArrayValuesRemainingStack.Num() > 0);
+	check(ArrayValuesRemainingStack.Top() > 0);
+}
+
+void FJsonArchiveInputFormatterEx::LeaveArrayElement()
+{
+	ValueStack.Pop();
+	ArrayValuesRemainingStack.Top()--;
+}
+
+void FJsonArchiveInputFormatterEx::EnterStream()
+{
+	int32 NumElements = 0;
+	EnterArray(NumElements);
+}
+
+void FJsonArchiveInputFormatterEx::LeaveStream()
+{
+	LeaveArray();
+}
+
+void FJsonArchiveInputFormatterEx::EnterStreamElement()
+{
+}
+
+void FJsonArchiveInputFormatterEx::LeaveStreamElement()
+{
+	LeaveArrayElement();
+}
+
+void FJsonArchiveInputFormatterEx::EnterMap(int32& NumElements)
+{
+	TSharedPtr<FJsonValue>& Value = ValueStack.Top();
+	TSharedPtr<FJsonObject> Object = Value->AsObject();
+	MapIteratorStack.Add(Object->Values);
+	NumElements = Object->Values.Num();
+}
+
+void FJsonArchiveInputFormatterEx::LeaveMap()
+{
+	MapIteratorStack.Pop();
+}
+
+void FJsonArchiveInputFormatterEx::EnterMapElement(FString& OutName)
+{
+	OutName = UnescapeFieldName(*MapIteratorStack.Top()->Key);
+	ValueStack.Add(MapIteratorStack.Top()->Value);
+}
+
+void FJsonArchiveInputFormatterEx::LeaveMapElement()
+{
+	ValueStack.Pop();
+	++MapIteratorStack.Top();
+}
+
+void FJsonArchiveInputFormatterEx::EnterAttributedValue()
+{
+	TSharedPtr<FJsonValue>& Value = ValueStack.Top();
+	const TSharedPtr<FJsonObject>* ObjectPtr = nullptr;
+	if (Value->TryGetObject(ObjectPtr))
+	{
+		FJsonObject& ObjectRef = **ObjectPtr;
+
+		TSharedPtr<FJsonValue> Field = ObjectRef.TryGetField(EscapeFieldName(TEXT("_Value")));
+		if (Field.IsValid())
+		{
+			ObjectStack.Emplace(*ObjectPtr, ValueStack.Num());
+			return;
+		}
+	}
+
+	ObjectStack.Emplace(nullptr, ValueStack.Num());
+}
+
+void FJsonArchiveInputFormatterEx::EnterAttribute(FArchiveFieldName AttributeName)
+{
+	TSharedPtr<FJsonObject>& NewObject = ObjectStack.Top().JsonObject;
+	TSharedPtr<FJsonValue> Field = NewObject->TryGetField(EscapeFieldName(*FString::Printf(TEXT("_%s"), AttributeName.Name)));
+	check(Field.IsValid());
+	ValueStack.Add(Field);
+}
+
+void FJsonArchiveInputFormatterEx::EnterAttributedValueValue()
+{
+	if (TSharedPtr<FJsonObject>& Object = ObjectStack.Top().JsonObject)
+	{
+		TSharedPtr<FJsonValue> Field = Object->TryGetField(EscapeFieldName(TEXT("_Value")));
+		checkSlow(Field);
+		ValueStack.Add(Field);
+	}
+	else
+	{
+		ValueStack.Add(CopyTemp(ValueStack.Top()));
+	}
+}
+
+bool FJsonArchiveInputFormatterEx::TryEnterAttributedValueValue()
+{
+	TSharedPtr<FJsonValue>& Value = ValueStack.Top();
+	const TSharedPtr<FJsonObject>* ObjectPtr = nullptr;
+	if (Value->TryGetObject(ObjectPtr))
+	{
+		FJsonObject& ObjectRef = **ObjectPtr;
+
+		TSharedPtr<FJsonValue> Field = ObjectRef.TryGetField(EscapeFieldName(TEXT("_Value")));
+		if (Field.IsValid())
+		{
+			ObjectStack.Emplace(*ObjectPtr, ValueStack.Num());
+			ValueStack.Add(Field);
+			return true;
+		}
+	}
+
+	return false;
+}
+
+void FJsonArchiveInputFormatterEx::LeaveAttribute()
+{
+	ValueStack.Pop();
+}
+
+void FJsonArchiveInputFormatterEx::LeaveAttributedValue()
+{
+	// check(ValueStack.Num() == ObjectStack.Top().ValueCountOnCreation);
+	ObjectStack.Pop();
+}
+
+bool FJsonArchiveInputFormatterEx::TryEnterAttribute(FArchiveFieldName AttributeName, bool bEnterWhenSaving)
+{
+	if (TSharedPtr<FJsonObject>& Object = ObjectStack.Top().JsonObject)
+	{
+		if (TSharedPtr<FJsonValue> Field = Object->TryGetField(EscapeFieldName(TEXT("_Value"))))
+		{
+			TSharedPtr<FJsonValue> Attribute = Object->TryGetField(EscapeFieldName(*FString::Printf(TEXT("_%s"), AttributeName.Name)));
+			if (Attribute.IsValid())
+			{
+				ValueStack.Add(Attribute);
+				return true;
+			}
+		}
+	}
+	return false;
+}
+
+void FJsonArchiveInputFormatterEx::Serialize(uint8& Value)
+{
+	verify(ValueStack.Top()->TryGetNumber(Value));
+}
+
+void FJsonArchiveInputFormatterEx::Serialize(uint16& Value)
+{
+	verify(ValueStack.Top()->TryGetNumber(Value));
+}
+
+void FJsonArchiveInputFormatterEx::Serialize(uint32& Value)
+{
+	verify(ValueStack.Top()->TryGetNumber(Value));
+}
+
+void FJsonArchiveInputFormatterEx::Serialize(uint64& Value)
+{
+	verify(ValueStack.Top()->TryGetNumber(Value));
+}
+
+void FJsonArchiveInputFormatterEx::Serialize(int8& Value)
+{
+	verify(ValueStack.Top()->TryGetNumber(Value));
+}
+
+void FJsonArchiveInputFormatterEx::Serialize(int16& Value)
+{
+	verify(ValueStack.Top()->TryGetNumber(Value));
+}
+
+void FJsonArchiveInputFormatterEx::Serialize(int32& Value)
+{
+	verify(ValueStack.Top()->TryGetNumber(Value));
+}
+
+void FJsonArchiveInputFormatterEx::Serialize(int64& Value)
+{
+	verify(ValueStack.Top()->TryGetNumber(Value));
+}
+
+void FJsonArchiveInputFormatterEx::Serialize(float& Value)
+{
+	verify(ValueStack.Top()->TryGetNumber(Value));
+}
+
+void FJsonArchiveInputFormatterEx::Serialize(double& Value)
+{
+	verify(ValueStack.Top()->TryGetNumber(Value));
+}
+
+void FJsonArchiveInputFormatterEx::Serialize(bool& Value)
+{
+	Value = ValueStack.Top()->AsBool();
+}
+
+void FJsonArchiveInputFormatterEx::Serialize(FString& Value)
+{
+	// If the string we serialized was empty, this value will be a null object rather than a string, so we have to
+#if DO_CHECK
+	bool bSuccess =
+#endif
+		ValueStack.Top()->TryGetString(Value);
+	check(bSuccess || ValueStack.Top()->IsNull());
+	Value.RemoveFromStart(TEXT("String:"));
+}
+
+void FJsonArchiveInputFormatterEx::Serialize(FName& Value)
+{
+	FString StringValue = ValueStack.Top()->AsString();
+	Value = FName(*StringValue);
+}
+
+void FJsonArchiveInputFormatterEx::Serialize(UObject*& Value)
+{
+	if (IsNestedObject())
+	{
+		Value = LoadNestedObject();
+		return;
+	}
+	
+	FString StringValue;
+	if (ValueStack.Top()->TryGetString(StringValue))
+	{
+		FPackageIndex ObjectIndex;
+		LexFromString(ObjectIndex, *StringValue);
+		Value = ResolveObject(ObjectIndex);
+	}
+	else
+	{
+		Value = nullptr;
+	}
+}
+
+void FJsonArchiveInputFormatterEx::Serialize(FText& Value)
+{
+	// FStructuredArchive ChildArchive(*this);
+	// FText::SerializeText(ChildArchive.Open(), Value);
+	// ChildArchive.Close();
+}
+
+void FJsonArchiveInputFormatterEx::Serialize(FWeakObjectPtr& Value)
+{
+	UObject* Object = nullptr;
+	Serialize(Object);
+	Value = Object;
+}
+
+void FJsonArchiveInputFormatterEx::Serialize(FSoftObjectPtr& Value)
+{
+	// Special case for nested objects.
+	if (IsNestedObject())
+	{
+		Value = LoadNestedObject();
+	}
+	else
+	{
+		FSoftObjectPath Path;
+		Serialize(Path);
+		Value = Path;
+	}
+}
+
+void FJsonArchiveInputFormatterEx::Serialize(FSoftObjectPath& Value)
+{
+	FString StringValue;
+	const auto& Prefix = TEXT("Object:");
+	if (ValueStack.Top()->TryGetString(StringValue) && StringValue.StartsWith(Prefix))
+	{
+		Value.SetPath(*StringValue + UE_ARRAY_COUNT(Prefix) - 1);
+	}
+	else
+	{
+		Value.Reset();
+	}
+}
+
+void FJsonArchiveInputFormatterEx::Serialize(FLazyObjectPtr& Value)
+{
+	FString StringValue;
+	const auto& Prefix = TEXT("Lazy:");
+	if (ValueStack.Top()->TryGetString(StringValue) && StringValue.StartsWith(Prefix))
+	{
+		FUniqueObjectGuid Guid;
+		Guid.FromString(*StringValue + UE_ARRAY_COUNT(Prefix) - 1);
+		Value = Guid;
+	}
+	else
+	{
+		Value.Reset();
+	}
+}
+
+void FJsonArchiveInputFormatterEx::Serialize(FObjectPtr& Value)
+{
+	UObject* Object;
+	Serialize(Object);
+	Value = Object;
+}
+
+void FJsonArchiveInputFormatterEx::Serialize(TArray<uint8>& Data)
+{
+	const FJsonValue& Value = *ValueStack.Top();
+	if(Value.Type == EJson::String)
+	{
+		// Single line string
+		FString RawData = ValueStack.Top()->AsString();
+		ensure(RawData.RemoveFromStart(TEXT("Base64:")));
+		FBase64::Decode(RawData, Data);
+	}
+	else if(Value.Type == EJson::Object)
+	{
+		// Multi-line string
+		FJsonObject& Object = *Value.AsObject();
+
+		// Read the digest
+		TSharedPtr<FJsonValue> DigestField = Object.TryGetField(TEXT("Digest"));
+		checkf(DigestField.IsValid() && DigestField->Type == EJson::String, TEXT("Missing or invalid 'Digest' field for raw data"));
+		FString Digest = DigestField->AsString();
+
+		// Read the base64 array
+		TSharedPtr<FJsonValue> Base64Field = Object.TryGetField(TEXT("Base64"));
+		checkf(Base64Field.IsValid() && Base64Field->Type == EJson::Array, TEXT("Missing or invalid 'Base64' field for raw data"));
+		const TArray<TSharedPtr<FJsonValue>>& Base64Array = Base64Field->AsArray();
+
+		// Parse the digest
+		uint8 ExpectedDigest[FSHA1::DigestSize];
+		verify(FString::ToHexBlob(DigestField->AsString(), ExpectedDigest, sizeof(ExpectedDigest)));
+
+		// Get the size of the encoded data
+		uint32 DecodedSize = 0;
+		for(const TSharedPtr<FJsonValue>& Base64Line : Base64Array)
+		{
+			DecodedSize += FBase64::GetDecodedDataSize(Base64Line->AsString());
+		}
+
+		// Read the encoded data
+		Data.SetNum(DecodedSize);
+
+		// Read each line
+		uint32 DecodedPos = 0;
+		for(const TSharedPtr<FJsonValue>& Base64Line : Base64Array)
+		{
+			FString Base64String = Base64Line->AsString();
+			verify(FBase64::Decode(*Base64String, Base64String.Len(), Data.GetData() + DecodedPos));
+			DecodedPos += FBase64::GetDecodedDataSize(Base64String);
+		}
+
+		// Make sure the digest matches
+		uint8 ActualDigest[FSHA1::DigestSize];
+		FSHA1::HashBuffer(Data.GetData(), Data.Num(), ActualDigest);
+		checkf(FMemory::Memcmp(ExpectedDigest, ActualDigest, FSHA1::DigestSize) == 0, TEXT("Digest does not match for raw data. Check that this file was merged correctly."));
+	}
+	else
+	{
+		checkf(false, TEXT("Invalid value type for raw data"));
+	}
+}
+
+void FJsonArchiveInputFormatterEx::Serialize(void* Data, uint64 DataSize)
+{
+	TArray<uint8> Buffer;
+	Serialize(Buffer);
+	check(Buffer.Num() == DataSize);
+	memcpy(Data, Buffer.GetData(), DataSize);
+}
+
+bool FJsonArchiveInputFormatterEx::IsNestedObject()
+{
+	return TryEnterField(TEXT("Class"), false);
+}
+
+UObject* FJsonArchiveInputFormatterEx::LoadNestedObject()
+{
+	FString ClassPathString;
+	Serialize(ClassPathString);
+	const FSoftObjectPath ClassPath(ClassPathString);
+	const TSoftObjectPtr<UClass> ClassPtr(ClassPath);
+	if (const UClass* const ObjectClass = ClassPtr.LoadSynchronous())
+	{
+		uint32 flags;
+		EnterField(TEXT("Flags"));
+		Serialize(flags);
+		const EObjectFlags ObjectFlags = static_cast<EObjectFlags>(flags);
+
+		FString OriginalObjectName;
+		EnterField(TEXT("Name"));
+		Serialize(OriginalObjectName);
+
+		// Name clashing prevention: in order to prevent name clashing, we let the new object be named with
+		// a unique name and rename the root field tag.
+		UObject* const NestedObject = NewObject<UObject>(RootObject, ObjectClass, *OriginalObjectName, ObjectFlags);
+
+		if (!NestedObject)
+		{
+			return nullptr;
+		}
+
+		{
+			EnterField(*OriginalObjectName);
+			FStructuredArchive ChildArchive(*this);
+			FStructuredArchiveRecord ObjectRecord = ChildArchive.Open().EnterRecord();
+			
+			for (FProperty* Property = NestedObject->GetClass()->PropertyLink; Property; Property = Property->PropertyLinkNext)
+			{
+				if (!Property || Property->HasAnyPropertyFlags(CPF_Transient))
+				{
+					continue;
+				}
+
+				FStructuredArchiveSlot PropertyField = ObjectRecord.EnterField(*Property->GetName());
+		
+				for (int32 Index = 0; Index < Property->ArrayDim; ++Index)
+				{
+					UObject* const Archetype = NestedObject->GetArchetype();
+			
+					void* const Data = Property->ContainerPtrToValuePtr<void>(NestedObject, Index);
+					const void* const Defaults = Property->ContainerPtrToValuePtrForDefaults<void>(Archetype->GetClass()
+						, Archetype
+						, Index);
+			
+					//if (!Property->Identical(Target, Default, UnderlyingArchive.GetPortFlags()))
+					{
+						FSerializedPropertyScope SerializedProperty(GetUnderlyingArchive(), Property);
+						Property->SerializeItem(PropertyField, Data, Defaults);
+					}
+				}
+			}
+			// FStructuredArchive ChildArchive(*this);
+			// FStructuredArchiveRecord ObjectRecord = ChildArchive.Open().EnterRecord();
+			// UE::JsonArchiveInputFormatter::Private::SerializeObject(ObjectRecord, NestedObject, TFunction<bool(FProperty*)>());
+			// LeaveRecord();
+			ChildArchive.Close();
+		}
+
+		NestedObject->PostLoad();
+
+		return NestedObject;
+	}
+	
+	return nullptr;
+}
+
+FString FJsonArchiveInputFormatterEx::EscapeFieldName(const TCHAR* Name)
+{
+	if(FCString::Stricmp(Name, TEXT("Base64")) == 0 || FCString::Stricmp(Name, TEXT("Digest")) == 0 || Name[0] == '_')
+	{
+		return FString::Printf(TEXT("_%s"), Name);
+	}
+	else
+	{
+		return Name;
+	}
+}
+
+FString FJsonArchiveInputFormatterEx::UnescapeFieldName(const TCHAR* Name)
+{
+	if(Name[0] == '_')
+	{
+		return Name + 1;
+	}
+	else
+	{
+		return Name;
+	}
+}
+
+EArchiveValueType FJsonArchiveInputFormatterEx::GetValueType(const FJsonValue& Value)
+{
+	switch (Value.Type)
+	{
+	case EJson::String:
+		if (Value.AsString().StartsWith(TEXT("Object:")))
+		{
+			return EArchiveValueType::Object;
+		}
+		else if (Value.AsString().StartsWith(TEXT("Base64:")))
+		{
+			return EArchiveValueType::RawData;
+		}
+		else
+		{
+			return EArchiveValueType::String;
+		}
+	case EJson::Number:
+		{
+			double Number = Value.AsNumber();
+			int64 NumberInt64 = (int64)Number;
+			if((double)NumberInt64 == Number)
+			{
+				if((int16)NumberInt64 == NumberInt64)
+				{
+					if((int8)NumberInt64 == NumberInt64)
+					{
+						return EArchiveValueType::Int8;
+					}
+					else
+					{
+						return EArchiveValueType::Int16;
+					}
+				}
+				else
+				{
+					if((int32)NumberInt64 == NumberInt64)
+					{
+						return EArchiveValueType::Int32;
+					}
+					else
+					{
+						return EArchiveValueType::Int64;
+					}
+				}
+			}
+			else
+			{
+				if((double)(float)Number == Number)
+				{
+					return EArchiveValueType::Float;
+				}
+				else
+				{
+					return EArchiveValueType::Double;
+				}
+			}
+		}
+	case EJson::Boolean:
+		return EArchiveValueType::Bool;
+	case EJson::Array:
+		return EArchiveValueType::Array;
+	case EJson::Object:
+		if(Value.AsObject()->HasField(TEXT("Base64")))
+		{
+			return EArchiveValueType::RawData;
+		}
+		else
+		{
+			return EArchiveValueType::Record;
+		}
+	case EJson::Null:
+		return EArchiveValueType::Object;
+	default:
+		checkf(false, TEXT("Unhandled value type in JSON archive (%d)"), (int)Value.Type);
+		return EArchiveValueType::None;
+	}
+}
+
+#endif
diff --git a/Engine/Plugins/Experimental/SerializationUtils/Source/JsonSerialization/Private/Formatters/JsonArchiveOutputFormatterEx.cpp b/Engine/Plugins/Experimental/SerializationUtils/Source/JsonSerialization/Private/Formatters/JsonArchiveOutputFormatterEx.cpp
new file mode 100644
--- /dev/null
+++ b/Engine/Plugins/Experimental/SerializationUtils/Source/JsonSerialization/Private/Formatters/JsonArchiveOutputFormatterEx.cpp
@@ -0,0 +1,643 @@
+// Copyright Epic Games, Inc. All Rights Reserved.
+
+#include "Formatters/JsonArchiveOutputFormatterEx.h"
+#include "Misc/Base64.h"
+#include "Misc/SecureHash.h"
+#include "Serialization/MemoryWriter.h"
+#include "UObject/LazyObjectPtr.h"
+#include "UObject/Object.h"
+#include "UObject/ObjectPtr.h"
+#include "UObject/SoftObjectPath.h"
+#include "UObject/SoftObjectPtr.h"
+#include "UObject/UnrealType.h"
+#include "UObject/WeakObjectPtr.h"
+
+#if WITH_TEXT_ARCHIVE_SUPPORT
+
+namespace UE::JsonArchiveOutputFormatter::Private
+{
+	void SerializeObject(FStructuredArchiveRecord& InRootRecord
+		, UObject* InObject
+		, const TFunction<bool(FProperty*)>& ShouldSkipProperty = TFunction<bool(FProperty*)>())
+	{
+		if (!InObject)
+		{
+			return;
+		}
+	
+		FStructuredArchiveRecord ObjectRecord = InRootRecord.EnterRecord(*InObject->GetName());
+		FArchive& UnderlyingArchive = InRootRecord.GetUnderlyingArchive();
+	
+		for (FProperty* Property = InObject->GetClass()->PropertyLink; Property; Property = Property->PropertyLinkNext)
+		{
+			if (!Property || Property->HasAnyPropertyFlags(CPF_Transient) || (ShouldSkipProperty && ShouldSkipProperty(Property)))
+			{
+				continue;
+			}
+
+			FStructuredArchiveSlot PropertyField = ObjectRecord.EnterField(*Property->GetName());
+		
+			for (int32 Index = 0; Index < Property->ArrayDim; ++Index)
+			{
+				UObject* const Archetype = InObject->GetArchetype();
+			
+				void* const Data = Property->ContainerPtrToValuePtr<void>(InObject, Index);
+				const void* const Defaults = Property->ContainerPtrToValuePtrForDefaults<void>(Archetype->GetClass()
+					, Archetype
+					, Index);
+			
+				//if (!Property->Identical(Target, Default, UnderlyingArchive.GetPortFlags()))
+				{
+					FSerializedPropertyScope SerializedProperty(UnderlyingArchive, Property);
+					Property->SerializeItem(PropertyField, Data, Defaults);
+				}
+			}
+		}
+	}
+};
+
+FJsonArchiveOutputFormatterEx::FJsonArchiveOutputFormatterEx(FArchive& InInner) 
+	: Inner(InInner)
+	, Newline(LINE_TERMINATOR_ANSI, UE_ARRAY_COUNT(LINE_TERMINATOR_ANSI) - 1)
+{
+	Inner.SetIsTextFormat(true);
+}
+
+FJsonArchiveOutputFormatterEx::~FJsonArchiveOutputFormatterEx()
+{
+}
+
+FArchive& FJsonArchiveOutputFormatterEx::GetUnderlyingArchive()
+{
+	return Inner;
+}
+
+bool FJsonArchiveOutputFormatterEx::HasDocumentTree() const
+{
+	return true;
+}
+
+void FJsonArchiveOutputFormatterEx::EnterRecord()
+{
+	WriteOptionalComma();
+	WriteOptionalNewline();
+	Write("{");
+	Newline.Add('\t');
+	bNeedsNewline = true;
+	TextStartPosStack.Push(Inner.Tell());
+}
+
+void FJsonArchiveOutputFormatterEx::LeaveRecord()
+{
+	Newline.Pop(EAllowShrinking::No);
+	if (TextStartPosStack.Pop() == Inner.Tell())
+	{
+		bNeedsNewline = false;
+	}
+	WriteOptionalNewline();
+	Write("}");
+	bNeedsComma = true;
+	bNeedsNewline = true;
+}
+
+void FJsonArchiveOutputFormatterEx::EnterField(FArchiveFieldName Name)
+{
+	WriteOptionalComma();
+	WriteOptionalNewline();
+	WriteFieldName(Name.Name);
+}
+
+void FJsonArchiveOutputFormatterEx::LeaveField()
+{
+	bNeedsComma = true;
+	bNeedsNewline = true;
+}
+
+bool FJsonArchiveOutputFormatterEx::TryEnterField(FArchiveFieldName Name, bool bEnterWhenSaving)
+{
+	if (bEnterWhenSaving)
+	{
+		EnterField(Name);
+	}
+	return bEnterWhenSaving;
+}
+
+void FJsonArchiveOutputFormatterEx::EnterArray(int32& NumElements)
+{
+	EnterStream();
+}
+
+void FJsonArchiveOutputFormatterEx::LeaveArray()
+{
+	LeaveStream();
+}
+
+void FJsonArchiveOutputFormatterEx::EnterArrayElement()
+{
+	EnterStreamElement();
+}
+
+void FJsonArchiveOutputFormatterEx::LeaveArrayElement()
+{
+	LeaveStreamElement();
+}
+
+void FJsonArchiveOutputFormatterEx::EnterStream()
+{
+	WriteOptionalComma();
+	WriteOptionalNewline();
+	Write("[");
+	Newline.Add('\t');
+	bNeedsNewline = true;
+	TextStartPosStack.Push(Inner.Tell());
+}
+
+void FJsonArchiveOutputFormatterEx::LeaveStream()
+{
+	Newline.Pop(EAllowShrinking::No);
+	if (TextStartPosStack.Pop() == Inner.Tell())
+	{
+		bNeedsNewline = false;
+	}
+	WriteOptionalNewline();
+	Write("]");
+	bNeedsComma = true;
+	bNeedsNewline = true;
+}
+
+void FJsonArchiveOutputFormatterEx::EnterStreamElement()
+{
+	WriteOptionalComma();
+	WriteOptionalNewline();
+}
+
+void FJsonArchiveOutputFormatterEx::LeaveStreamElement()
+{
+	bNeedsComma = true;
+	bNeedsNewline = true;
+}
+
+void FJsonArchiveOutputFormatterEx::EnterMap(int32& NumElements)
+{
+	EnterRecord();
+}
+
+void FJsonArchiveOutputFormatterEx::LeaveMap()
+{
+	LeaveRecord();
+}
+
+void FJsonArchiveOutputFormatterEx::EnterMapElement(FString& Name)
+{
+	EnterField(FArchiveFieldName(*Name));
+}
+
+void FJsonArchiveOutputFormatterEx::LeaveMapElement()
+{
+	LeaveField();
+}
+
+void FJsonArchiveOutputFormatterEx::EnterAttributedValue()
+{
+	NumAttributesStack.Push(0);
+}
+
+void FJsonArchiveOutputFormatterEx::EnterAttribute(FArchiveFieldName AttributeName)
+{
+	WriteOptionalComma();
+	WriteOptionalNewline();
+	WriteOptionalAttributedBlockOpening();
+	WriteOptionalComma();
+	WriteOptionalNewline();
+	checkf(FCString::Strcmp(AttributeName.Name, TEXT("Value")) != 0, TEXT("Attributes called 'Value' are reserved by the implementation"));
+	WriteFieldName(*FString::Printf(TEXT("_%s"), AttributeName.Name));
+	++NumAttributesStack.Top();
+}
+
+void FJsonArchiveOutputFormatterEx::LeaveAttribute()
+{
+	bNeedsComma = true;
+	bNeedsNewline = true;
+}
+
+void FJsonArchiveOutputFormatterEx::LeaveAttributedValue()
+{
+	WriteOptionalAttributedBlockClosing();
+	NumAttributesStack.Pop();
+	bNeedsComma = true;
+	bNeedsNewline = true;
+}
+
+void FJsonArchiveOutputFormatterEx::EnterAttributedValueValue()
+{
+	WriteOptionalComma();
+	WriteOptionalNewline();
+	WriteOptionalAttributedBlockValue();
+}
+
+bool FJsonArchiveOutputFormatterEx::TryEnterAttributedValueValue()
+{
+	return false;
+}
+
+bool FJsonArchiveOutputFormatterEx::TryEnterAttribute(FArchiveFieldName AttributeName, bool bEnterWhenSaving)
+{
+	if (bEnterWhenSaving)
+	{
+		EnterAttribute(AttributeName);
+	}
+	return bEnterWhenSaving;
+}
+
+void FJsonArchiveOutputFormatterEx::Serialize(uint8& Value)
+{
+	WriteValue(LexToString(Value));
+}
+
+void FJsonArchiveOutputFormatterEx::Serialize(uint16& Value)
+{
+	WriteValue(LexToString(Value));
+}
+
+void FJsonArchiveOutputFormatterEx::Serialize(uint32& Value)
+{
+	WriteValue(LexToString(Value));
+}
+
+void FJsonArchiveOutputFormatterEx::Serialize(uint64& Value)
+{
+	WriteValue(LexToString(Value));
+}
+
+void FJsonArchiveOutputFormatterEx::Serialize(int8& Value)
+{
+	WriteValue(LexToString(Value));
+}
+
+void FJsonArchiveOutputFormatterEx::Serialize(int16& Value)
+{
+	WriteValue(LexToString(Value));
+}
+
+void FJsonArchiveOutputFormatterEx::Serialize(int32& Value)
+{
+	WriteValue(LexToString(Value));
+}
+
+void FJsonArchiveOutputFormatterEx::Serialize(int64& Value)
+{
+	WriteValue(LexToString(Value));
+}
+
+void FJsonArchiveOutputFormatterEx::Serialize(float& Value)
+{
+	if((float)(int)Value == Value)
+	{
+		WriteValue(LexToString((int)Value));
+	}
+	else
+	{
+		FString String = FString::Printf(TEXT("%.17g"), Value);
+#if DO_GUARD_SLOW
+		float RoundTripped;
+		LexFromString(RoundTripped, *String);
+		check(RoundTripped == Value);
+#endif
+		WriteValue(String);
+	}
+}
+
+void FJsonArchiveOutputFormatterEx::Serialize(double& Value)
+{
+	if((double)(int)Value == Value)
+	{
+		WriteValue(LexToString((int)Value));
+	}
+	else
+	{
+		FString String = FString::Printf(TEXT("%.17g"), Value);
+#if DO_GUARD_SLOW
+		double RoundTripped;
+		LexFromString(RoundTripped, *String);
+		check(RoundTripped == Value);
+#endif
+		WriteValue(String);
+	}
+}
+
+void FJsonArchiveOutputFormatterEx::Serialize(bool& Value)
+{
+	WriteValue(LexToString(Value));
+}
+
+void FJsonArchiveOutputFormatterEx::Serialize(FString& Value)
+{
+	// Insert a "String:" prefix to prevent incorrect interpretation as another explicit type
+	if (Value.StartsWith(TEXT("Object:")) || Value.StartsWith(TEXT("String:")) || Value.StartsWith(TEXT("Base64:")))
+	{
+		SerializeStringInternal(FString::Printf(TEXT("String:%s"), *Value));
+	}
+	else
+	{
+		SerializeStringInternal(Value);
+	}
+}
+
+void FJsonArchiveOutputFormatterEx::Serialize(FName& Value)
+{
+	SerializeStringInternal(*Value.ToString());
+}
+
+void FJsonArchiveOutputFormatterEx::Serialize(UObject*& Value)
+{
+	if (!bSerializeObjectsInPlace)
+	{
+		if (Value != nullptr && IsObjectAllowed(Value))
+		{
+			FPackageIndex ObjectIndex = ObjectIndicesMap->FindChecked(Value);
+			SerializeStringInternal(LexToString(ObjectIndex));
+		}
+		else
+		{
+			WriteValue(TEXT("null"));
+		}
+
+		return;
+	}
+
+	FStructuredArchive ChildArchive(*this);
+	FStructuredArchiveRecord ObjectRecord = ChildArchive.Open().EnterRecord();
+	// For current record, add attribute for object type
+	const TSoftObjectPtr<UClass> ObjectClass =  Value->GetClass();
+
+	EnterField(TEXT("Class"));
+	FString ObjectClassString = ObjectClass.ToSoftObjectPath().ToString();
+	Serialize(ObjectClassString);
+	LeaveField();
+
+	EnterField(TEXT("Name"));
+	FString Name = Value->GetName();
+	Serialize(Name);
+	LeaveField();
+
+	EnterField(TEXT("Flags"));
+	uint32 Flags = static_cast<uint32>(Value->GetFlags());
+	Serialize(Flags);
+	LeaveField();
+	
+	UE::JsonArchiveOutputFormatter::Private::SerializeObject(ObjectRecord, Value);
+	ChildArchive.Close();
+}
+
+void FJsonArchiveOutputFormatterEx::Serialize(FText& Value)
+{
+	// FStructuredArchive ChildArchive(*this);
+	// FText::SerializeText(ChildArchive.Open(), Value);
+	// ChildArchive.Close();
+}
+
+void FJsonArchiveOutputFormatterEx::Serialize(FWeakObjectPtr& Value)
+{
+	UObject* Ptr = Value.IsValid() ? Value.Get() : nullptr;
+	Serialize(Ptr);
+}
+
+void FJsonArchiveOutputFormatterEx::Serialize(FSoftObjectPtr& Value)
+{
+	FSoftObjectPath Path = Value.ToSoftObjectPath();
+	Serialize(Path);
+}
+
+void FJsonArchiveOutputFormatterEx::Serialize(FSoftObjectPath& Value)
+{
+	if (Value.IsValid())
+	{
+		SerializeStringInternal(FString::Printf(TEXT("Object:%s"), *Value.ToString()));
+	}
+	else
+	{
+		WriteValue(TEXT("null"));
+	}
+}
+
+void FJsonArchiveOutputFormatterEx::Serialize(FLazyObjectPtr& Value)
+{
+	if (Value.IsValid() && IsObjectAllowed(Value.Get()))
+	{
+		SerializeStringInternal(FString::Printf(TEXT("Lazy:%s"), *Value.GetUniqueID().ToString()));
+	}
+	else
+	{
+		WriteValue(TEXT("null"));
+	}
+}
+
+void FJsonArchiveOutputFormatterEx::Serialize(FObjectPtr& Value)
+{
+	UObject* Object = Value.Get();
+	Serialize(Object);
+}
+
+void FJsonArchiveOutputFormatterEx::Serialize(TArray<uint8>& Data)
+{
+	Serialize(Data.GetData(), Data.Num());
+}
+
+void FJsonArchiveOutputFormatterEx::Serialize(void* Data, uint64 DataSize)
+{
+	static const int32 MaxLineChars = 120;
+	static const int32 MaxLineBytes = FBase64::GetMaxDecodedDataSize(MaxLineChars);
+
+	if(DataSize < MaxLineBytes)
+	{
+		// Encode the data on a single line. No need for hashing; intra-line merge conflicts are rare.
+		WriteValue(FString::Printf(TEXT("\"Base64:%s\""), *FBase64::Encode((const uint8*)Data, DataSize)));
+	}
+	else
+	{
+		// Encode the data as a record containing a digest and array of base-64 encoded lines
+		EnterRecord();
+		Inner.Serialize(Newline.GetData(), Newline.Num());
+
+		// Compute a SHA digest for the raw data, so we can check if it's corrupted
+		uint8 Digest[FSHA1::DigestSize];
+		FSHA1::HashBuffer(Data, DataSize, Digest);
+
+		// Convert the hash to a string
+		ANSICHAR DigestString[(FSHA1::DigestSize * 2) + 1];
+		for(int32 Idx = 0; Idx < UE_ARRAY_COUNT(Digest); Idx++)
+		{
+			static const ANSICHAR HexDigits[] = "0123456789abcdef";
+			DigestString[(Idx * 2) + 0] = HexDigits[Digest[Idx] >> 4];
+			DigestString[(Idx * 2) + 1] = HexDigits[Digest[Idx] & 15];
+		}
+		DigestString[UE_ARRAY_COUNT(DigestString) - 1] = 0;
+
+		// Write the digest
+		Write("\"Digest\": \"");
+		Write(DigestString);
+		Write("\",");
+		Inner.Serialize(Newline.GetData(), Newline.Num());
+
+		// Write the base64 data
+		Write("\"Base64\": ");
+		for(uint64 DataPos = 0; DataPos < DataSize; DataPos += MaxLineBytes)
+		{
+			Write((DataPos > 0)? ',' : '[');
+			Inner.Serialize(Newline.GetData(), Newline.Num());
+			Write("\t\"");
+
+			ANSICHAR LineData[MaxLineChars + 1];
+			uint64 NumLineChars = FBase64::Encode((const uint8*)Data + DataPos, FMath::Min<uint64>(DataSize - DataPos, MaxLineBytes), LineData);
+			Inner.Serialize(LineData, NumLineChars);
+
+			Write("\"");
+		}
+
+		// Close the array
+		Inner.Serialize(Newline.GetData(), Newline.Num());
+		Write(']');
+		bNeedsNewline = true;
+
+		// Close the record
+		LeaveRecord();
+	}
+}
+
+void FJsonArchiveOutputFormatterEx::Write(ANSICHAR Character)
+{
+	Inner.Serialize((void*)&Character, 1);
+}
+
+void FJsonArchiveOutputFormatterEx::Write(const ANSICHAR* Text)
+{
+	Inner.Serialize((void*)Text, TCString<ANSICHAR>::Strlen(Text));
+}
+
+void FJsonArchiveOutputFormatterEx::Write(const FString& Text)
+{
+	Write(TCHAR_TO_UTF8(*Text));
+}
+
+void FJsonArchiveOutputFormatterEx::WriteFieldName(const TCHAR* Name)
+{
+	if(FCString::Stricmp(Name, TEXT("Base64")) == 0 || FCString::Stricmp(Name, TEXT("Digest")) == 0)
+	{
+		Write(FString::Printf(TEXT("\"_%s\": "), Name));
+	}
+	else if(Name[0] == '_')
+	{
+		Write(FString::Printf(TEXT("\"_%s\": "), Name));
+	}
+	else
+	{
+		Write(FString::Printf(TEXT("\"%s\": "), Name));
+	}
+}
+
+void FJsonArchiveOutputFormatterEx::WriteValue(const FString& Text)
+{
+	Write(Text);
+}
+
+void FJsonArchiveOutputFormatterEx::WriteOptionalComma()
+{
+	if (bNeedsComma)
+	{
+		Write(',');
+		bNeedsComma = false;
+	}
+}
+
+void FJsonArchiveOutputFormatterEx::WriteOptionalNewline()
+{
+	if (bNeedsNewline)
+	{
+		Inner.Serialize(Newline.GetData(), Newline.Num());
+		bNeedsNewline = false;
+	}
+}
+
+void FJsonArchiveOutputFormatterEx::WriteOptionalAttributedBlockOpening()
+{
+	if (NumAttributesStack.Top() == 0)
+	{
+		Write('{');
+		Newline.Add('\t');
+		bNeedsNewline = true;
+	}
+}
+
+void FJsonArchiveOutputFormatterEx::WriteOptionalAttributedBlockValue()
+{
+	if (NumAttributesStack.Top() != 0)
+	{
+		WriteFieldName(TEXT("_Value"));
+	}
+}
+
+void FJsonArchiveOutputFormatterEx::WriteOptionalAttributedBlockClosing()
+{
+	if (NumAttributesStack.Top() != 0)
+	{
+		Newline.Pop(EAllowShrinking::No);
+		WriteOptionalNewline();
+		Write("}");
+		bNeedsComma                  = true;
+		bNeedsNewline                = true;
+	}
+}
+
+void FJsonArchiveOutputFormatterEx::SerializeStringInternal(const FString& String)
+{
+	FString Result = TEXT("\"");
+
+	// Escape the string characters
+	for (int32 Idx = 0; Idx < String.Len(); Idx++)
+	{
+		switch (String[Idx])
+		{
+		case '\"':
+			Result += "\\\"";
+			break;
+		case '\\':
+			Result += "\\\\";
+			break;
+		case '\b':
+			Result += "\\b";
+			break;
+		case '\f':
+			Result += "\\f";
+			break;
+		case '\n':
+			Result += "\\n";
+			break;
+		case '\r':
+			Result += "\\r";
+			break;
+		case '\t':
+			Result += "\\t";
+			break;
+		default:
+			if (String[Idx] <= 0x1f || String[Idx] >= 0x7f)
+			{
+				Result += FString::Printf(TEXT("\\u%04x"), String[Idx]);
+			}
+			else
+			{
+				Result.AppendChar(String[Idx]);
+			}
+			break;
+		}
+	}
+	Result += TEXT("\"");
+
+	WriteValue(Result);
+}
+
+bool FJsonArchiveOutputFormatterEx::IsObjectAllowed(UObject* InObject) const
+{
+	return ObjectIndicesMap && ObjectIndicesMap->Contains(InObject);
+}
+
+#endif
diff --git a/Engine/Plugins/Experimental/SerializationUtils/Source/JsonSerialization/Private/JsonSerializationModule.cpp b/Engine/Plugins/Experimental/SerializationUtils/Source/JsonSerialization/Private/JsonSerializationModule.cpp
new file mode 100644
--- /dev/null
+++ b/Engine/Plugins/Experimental/SerializationUtils/Source/JsonSerialization/Private/JsonSerializationModule.cpp
@@ -0,0 +1,16 @@
+// Copyright Epic Games, Inc. All Rights Reserved.
+
+#include "JsonSerializationModule.h"
+#include "Modules/ModuleManager.h"
+
+DEFINE_LOG_CATEGORY(LogJsonSerialization);
+
+void FJsonSerializationModule::StartupModule()
+{
+}
+
+void FJsonSerializationModule::ShutdownModule()
+{
+}
+
+IMPLEMENT_MODULE(FJsonSerializationModule, JsonSerialization)
\ No newline at end of file
diff --git a/Engine/Plugins/Experimental/SerializationUtils/Source/JsonSerialization/Private/JsonSerializationModule.h b/Engine/Plugins/Experimental/SerializationUtils/Source/JsonSerialization/Private/JsonSerializationModule.h
new file mode 100644
--- /dev/null
+++ b/Engine/Plugins/Experimental/SerializationUtils/Source/JsonSerialization/Private/JsonSerializationModule.h
@@ -0,0 +1,17 @@
+// Copyright Epic Games, Inc. All Rights Reserved.
+
+#pragma once
+
+#include "CoreMinimal.h"
+#include "Modules/ModuleInterface.h"
+
+DECLARE_LOG_CATEGORY_EXTERN(LogJsonSerialization, Log, All);
+
+class FJsonSerializationModule : public IModuleInterface
+{
+public:
+	//~ Begin IModuleInterface
+	virtual void StartupModule() override;
+	virtual void ShutdownModule() override;
+	//~ End IModuleInterface
+};
diff --git a/Engine/Plugins/Experimental/SerializationUtils/Source/JsonSerialization/Public/Formatters/JsonArchiveInputFormatterEx.h b/Engine/Plugins/Experimental/SerializationUtils/Source/JsonSerialization/Public/Formatters/JsonArchiveInputFormatterEx.h
new file mode 100644
--- /dev/null
+++ b/Engine/Plugins/Experimental/SerializationUtils/Source/JsonSerialization/Public/Formatters/JsonArchiveInputFormatterEx.h
@@ -0,0 +1,126 @@
+// Copyright Epic Games, Inc. All Rights Reserved.
+
+#pragma once
+
+#include "Containers/Array.h"
+#include "Containers/Map.h"
+#include "Containers/UnrealString.h"
+#include "CoreMinimal.h"
+#include "HAL/Platform.h"
+#include "Serialization/StructuredArchiveFormatter.h"
+#include "Serialization/StructuredArchiveNameHelpers.h"
+#include "Templates/Function.h"
+#include "Templates/SharedPointer.h"
+#include "UObject/ObjectResource.h"
+
+class FArchive;
+class FJsonObject;
+class FJsonValue;
+class FName;
+class FPackageIndex;
+class FText;
+class UObject;
+struct FLazyObjectPtr;
+struct FObjectPtr;
+struct FSoftObjectPath;
+struct FSoftObjectPtr;
+struct FWeakObjectPtr;
+
+#if WITH_TEXT_ARCHIVE_SUPPORT
+
+class JSONSERIALIZATION_API FJsonArchiveInputFormatterEx final : public FStructuredArchiveFormatter
+{
+public:
+	FJsonArchiveInputFormatterEx(FArchive& InInner, UObject* InRootObject, TFunction<UObject* (const FPackageIndex)> InResolveObject = nullptr);
+	virtual ~FJsonArchiveInputFormatterEx() override;
+
+	virtual FArchive& GetUnderlyingArchive() override;
+	virtual FStructuredArchiveFormatter* CreateSubtreeReader() override;
+
+	virtual bool HasDocumentTree() const override;
+
+	virtual void EnterRecord() override;
+	virtual void LeaveRecord() override;
+	virtual void EnterField(FArchiveFieldName Name) override;
+	virtual void LeaveField() override;
+	virtual bool TryEnterField(FArchiveFieldName Name, bool bEnterWhenSaving) override;
+
+	virtual void EnterArray(int32& NumElements) override;
+	virtual void LeaveArray() override;
+	virtual void EnterArrayElement() override;
+	virtual void LeaveArrayElement() override;
+
+	virtual void EnterStream() override;
+	virtual void LeaveStream() override;
+	virtual void EnterStreamElement() override;
+	virtual void LeaveStreamElement() override;
+
+	virtual void EnterMap(int32& NumElements) override;
+	virtual void LeaveMap() override;
+	virtual void EnterMapElement(FString& Name) override;
+	virtual void LeaveMapElement() override;
+
+	virtual void EnterAttributedValue() override;
+	virtual void EnterAttribute(FArchiveFieldName AttributeName) override;
+	virtual void EnterAttributedValueValue() override;
+	virtual void LeaveAttribute() override;
+	virtual void LeaveAttributedValue() override;
+	virtual bool TryEnterAttribute(FArchiveFieldName AttributeName, bool bEnterWhenSavin) override;
+	virtual bool TryEnterAttributedValueValue() override;
+
+	virtual void Serialize(uint8& Value) override;
+	virtual void Serialize(uint16& Value) override;
+	virtual void Serialize(uint32& Value) override;
+	virtual void Serialize(uint64& Value) override;
+	virtual void Serialize(int8& Value) override;
+	virtual void Serialize(int16& Value) override;
+	virtual void Serialize(int32& Value) override;
+	virtual void Serialize(int64& Value) override;
+	virtual void Serialize(float& Value) override;
+	virtual void Serialize(double& Value) override;
+	virtual void Serialize(bool& Value) override;
+	virtual void Serialize(FString& Value) override;
+	virtual void Serialize(FName& Value) override;
+	virtual void Serialize(UObject*& Value) override;
+	virtual void Serialize(FText& Value) override;
+	virtual void Serialize(FWeakObjectPtr& Value) override;
+	virtual void Serialize(FSoftObjectPtr& Value) override;
+	virtual void Serialize(FSoftObjectPath& Value) override;
+	virtual void Serialize(FLazyObjectPtr& Value) override;
+	virtual void Serialize(FObjectPtr& Value) override;
+	virtual void Serialize(TArray<uint8>& Value) override;
+	virtual void Serialize(void* Data, uint64 DataSize) override;
+
+private:
+	bool IsNestedObject();
+	UObject* LoadNestedObject();
+	
+	FArchive& Inner;
+
+	struct FObjectRecord
+	{
+		FObjectRecord(TSharedPtr<FJsonObject> InJsonObject, int64 InValueCount)
+			: JsonObject(InJsonObject)
+			, ValueCountOnCreation(InValueCount)
+		{
+
+		}
+
+		TSharedPtr<FJsonObject> JsonObject;
+		int64 ValueCountOnCreation;	// For debugging purposes, so we can ensure all values have been consumed
+	};
+
+	TFunction<UObject* (const FPackageIndex)> ResolveObject;
+	UObject* RootObject;
+	TArray<FObjectRecord> ObjectStack;
+	TArray<TSharedPtr<FJsonValue>> ValueStack;
+	TArray<TMap<FString, TSharedPtr<FJsonValue>>::TIterator> MapIteratorStack;
+	TArray<int32> ArrayValuesRemainingStack;
+
+	static FString EscapeFieldName(const TCHAR* Name);
+	static FString UnescapeFieldName(const TCHAR* Name);
+
+	static EArchiveValueType GetValueType(const FJsonValue &Value);
+};
+
+#endif
diff --git a/Engine/Plugins/Experimental/SerializationUtils/Source/JsonSerialization/Public/Formatters/JsonArchiveOutputFormatterEx.h b/Engine/Plugins/Experimental/SerializationUtils/Source/JsonSerialization/Public/Formatters/JsonArchiveOutputFormatterEx.h
new file mode 100644
--- /dev/null
+++ b/Engine/Plugins/Experimental/SerializationUtils/Source/JsonSerialization/Public/Formatters/JsonArchiveOutputFormatterEx.h
@@ -0,0 +1,131 @@
+// Copyright Epic Games, Inc. All Rights Reserved.
+
+#pragma once
+
+#include "Containers/Array.h"
+#include "Containers/Map.h"
+#include "Containers/UnrealString.h"
+#include "CoreMinimal.h"
+#include "HAL/Platform.h"
+#include "Serialization/StructuredArchiveFormatter.h"
+#include "Serialization/StructuredArchiveNameHelpers.h"
+#include "UObject/ObjectResource.h"
+
+class FArchive;
+class FName;
+class FPackageIndex;
+class FText;
+class UObject;
+struct FLazyObjectPtr;
+struct FObjectPtr;
+struct FSoftObjectPath;
+struct FSoftObjectPtr;
+struct FWeakObjectPtr;
+
+#if WITH_TEXT_ARCHIVE_SUPPORT
+
+class JSONSERIALIZATION_API FJsonArchiveOutputFormatterEx : public FStructuredArchiveFormatter
+{
+public:
+	FJsonArchiveOutputFormatterEx(FArchive& InInner);
+	virtual ~FJsonArchiveOutputFormatterEx() override;
+
+	void SerializeObjectsInPlace(bool bEnabled) { bSerializeObjectsInPlace = bEnabled; }
+	void SerializeSoftObjectsInPlace(bool bEnabled) { bSerializeSoftObjectsInPlace = bEnabled; }
+	
+	virtual FArchive& GetUnderlyingArchive() override;
+
+	virtual bool HasDocumentTree() const override;
+
+	virtual void EnterRecord() override;
+	virtual void LeaveRecord() override;
+	virtual void EnterField(FArchiveFieldName Name) override;
+	virtual void LeaveField() override;
+	virtual bool TryEnterField(FArchiveFieldName Name, bool bEnterWhenSaving) override;
+
+	virtual void EnterArray(int32& NumElements) override;
+	virtual void LeaveArray() override;
+	virtual void EnterArrayElement() override;
+	virtual void LeaveArrayElement() override;
+
+	virtual void EnterStream() override;
+	virtual void LeaveStream() override;
+	virtual void EnterStreamElement() override;
+	virtual void LeaveStreamElement() override;
+
+	virtual void EnterMap(int32& NumElements) override;
+	virtual void LeaveMap() override;
+	virtual void EnterMapElement(FString& Name) override;
+	virtual void LeaveMapElement() override;
+
+	virtual void EnterAttributedValue() override;
+	virtual void EnterAttribute(FArchiveFieldName AttributeName) override;
+	virtual void EnterAttributedValueValue() override;
+	virtual void LeaveAttribute() override;
+	virtual void LeaveAttributedValue() override;
+	virtual bool TryEnterAttribute(FArchiveFieldName AttributeName, bool bEnterWhenSaving) override;
+
+	virtual bool TryEnterAttributedValueValue() override;
+
+	virtual void Serialize(uint8& Value) override;
+	virtual void Serialize(uint16& Value) override;
+	virtual void Serialize(uint32& Value) override;
+	virtual void Serialize(uint64& Value) override;
+	virtual void Serialize(int8& Value) override;
+	virtual void Serialize(int16& Value) override;
+	virtual void Serialize(int32& Value) override;
+	virtual void Serialize(int64& Value) override;
+	virtual void Serialize(float& Value) override;
+	virtual void Serialize(double& Value) override;
+	virtual void Serialize(bool& Value) override;
+	virtual void Serialize(FString& Value) override;
+	virtual void Serialize(FName& Value) override;
+	virtual void Serialize(UObject*& Value) override;
+	virtual void Serialize(FText& Value) override;
+	virtual void Serialize(FWeakObjectPtr& Value) override;
+	virtual void Serialize(FSoftObjectPtr& Value) override;
+	virtual void Serialize(FSoftObjectPath& Value) override;
+	virtual void Serialize(FLazyObjectPtr& Value) override;
+	virtual void Serialize(FObjectPtr& Value) override;
+	virtual void Serialize(TArray<uint8>& Value) override;
+	virtual void Serialize(void* Data, uint64 DataSize) override;
+
+	void SetObjectIndicesMap(const TMap<UObject*, FPackageIndex>* InObjectIndicesMap)
+	{
+		ObjectIndicesMap = InObjectIndicesMap;
+	}
+
+private:
+	FArchive& Inner;
+
+	const TMap<UObject*, FPackageIndex>* ObjectIndicesMap = nullptr;
+	bool bSerializeObjectsInPlace = false;
+	bool bSerializeSoftObjectsInPlace = false;
+	
+	TArray<ANSICHAR> Newline;
+	bool bNeedsComma   = false;
+	bool bNeedsNewline = false;
+
+	TArray<int32> NumAttributesStack;
+	TArray<int64> TextStartPosStack;
+
+	void Write(ANSICHAR Character);
+
+	void Write(const ANSICHAR* Text);
+	void Write(const FString& Text);
+
+	void WriteFieldName(const TCHAR* Name);
+	void WriteValue(const FString& Value);
+
+	void WriteOptionalComma();
+	void WriteOptionalNewline();
+	void WriteOptionalAttributedBlockOpening();
+	void WriteOptionalAttributedBlockValue();
+	void WriteOptionalAttributedBlockClosing();
+
+	void SerializeStringInternal(const FString& String);
+
+	bool IsObjectAllowed(UObject* InObject) const;
+};
+
+#endif
diff --git a/Engine/Plugins/Experimental/SerializationUtils/Source/XmlSerialization/Private/Backends/XmlStructSerializerBackend.cpp b/Engine/Plugins/Experimental/SerializationUtils/Source/XmlSerialization/Private/Backends/XmlStructSerializerBackend.cpp
new file mode 100644
--- /dev/null
+++ b/Engine/Plugins/Experimental/SerializationUtils/Source/XmlSerialization/Private/Backends/XmlStructSerializerBackend.cpp
@@ -0,0 +1,243 @@
+// Copyright Epic Games, Inc. All Rights Reserved.
+
+#include "Backends/XmlStructSerializerBackend.h"
+
+#include "XmlSerializationModule.h"
+#include "MaterialXFormat/PugiXML/pugixml.hpp"
+#include "UObject/EnumProperty.h"
+#include "UObject/PropertyPortFlags.h"
+#include "UObject/TextProperty.h"
+#include "UObject/UnrealType.h"
+#include "Utils/XmlUtils.h"
+#include "Utils/XmlWriter.h"
+
+class FXmlStructSerializerBackend::FImpl
+{
+public:
+	FImpl()
+	{
+		pugi::xml_node decl = XmlDoc.prepend_child(pugi::node_declaration);
+		decl.append_attribute("version") = "1.0";
+		decl.append_attribute("encoding") = "utf-8";
+
+		XmlElement = XmlDoc.append_child();
+		XmlElement.set_name("StructFileDocument");
+		XmlElement.append_attribute(("Version")) = "1.0";
+	}
+	
+	pugi::xml_document XmlDoc;
+	pugi::xml_node XmlElement;
+	TUniquePtr<FXmlWriter> XmlWriter;
+};
+
+FXmlStructSerializerBackend::FXmlStructSerializerBackend(FArchive& InArchive,
+                                                         const EStructSerializerBackendFlags InFlags)
+	: Impl(new FImpl)
+	  , Flags(InFlags)
+{
+	Impl->XmlWriter = MakeUnique<FXmlWriter>(&InArchive);
+}
+
+FXmlStructSerializerBackend::~FXmlStructSerializerBackend()
+{
+	delete Impl;
+}
+
+void FXmlStructSerializerBackend::SaveDocument(EXmlSerializationEncoding InEncoding)
+{
+	using namespace UE::XmlSerialization::Private;
+	const pugi::xml_encoding encoding = Utils::ToPugiEncoding(InEncoding);
+	Impl->XmlDoc.save(*Impl->XmlWriter, "  ", pugi::format_default, encoding);
+}
+
+void FXmlStructSerializerBackend::BeginArray(const FStructSerializerState& InState)
+{
+	pugi::xml_node ArrayElement = Impl->XmlElement.append_child(pugi::node_element);
+	ArrayElement.set_name("Array");
+	if (InState.ValueProperty)
+	{
+		ArrayElement.append_attribute("Name") = STRING_CAST_UE_TO_PUGI(*InState.ValueProperty->GetName());
+		ArrayElement.append_attribute("Type") = STRING_CAST_UE_TO_PUGI(*InState.ValueProperty->GetClass()->GetName());
+	}
+	if (InState.ValueType)
+	{
+		ArrayElement.append_attribute("ValueType") = STRING_CAST_UE_TO_PUGI(*InState.ValueType->GetName());
+	}
+	Impl->XmlElement = ArrayElement;
+}
+
+void FXmlStructSerializerBackend::BeginStructure(const FStructSerializerState& InState)
+{
+	pugi::xml_node StructElement = Impl->XmlElement.append_child(pugi::node_element);
+	StructElement.set_name("Struct");
+	if (InState.ValueProperty)
+	{
+		StructElement.append_attribute("Name") = STRING_CAST_UE_TO_PUGI(*InState.ValueProperty->GetName());
+		StructElement.append_attribute("Type") = STRING_CAST_UE_TO_PUGI(*InState.ValueProperty->GetClass()->GetName());
+	}
+	if (InState.ValueType)
+	{
+		StructElement.append_attribute("ValueType") = STRING_CAST_UE_TO_PUGI(*InState.ValueType->GetName());
+	}
+	Impl->XmlElement = StructElement;
+}
+
+
+void FXmlStructSerializerBackend::EndArray(const FStructSerializerState& /*State*/)
+{
+	// Check that current element is an Array?
+	Impl->XmlElement = Impl->XmlElement.parent();
+}
+
+
+void FXmlStructSerializerBackend::EndStructure(const FStructSerializerState& /*State*/)
+{
+	// Check that current element is a Struct?
+	Impl->XmlElement = Impl->XmlElement.parent();
+}
+
+void FXmlStructSerializerBackend::WriteComment(const FString& InComment)
+{
+	pugi::xml_node Child = Impl->XmlElement.append_child(pugi::node_comment);
+	Child.set_value(STRING_CAST_UE_TO_PUGI(*InComment));
+}
+
+void FXmlStructSerializerBackend::WriteProperty(const FStructSerializerState& InState, int32 InArrayIndex)
+{
+	// Append a property element in the current structure element.
+	pugi::xml_node PropertyElement = Impl->XmlElement.append_child(pugi::node_element);
+	PropertyElement.set_name("Property");
+	PropertyElement.append_attribute("Name") = STRING_CAST_UE_TO_PUGI(*InState.ValueProperty->GetName());
+	PropertyElement.append_attribute("Type") = STRING_CAST_UE_TO_PUGI(*InState.ValueProperty->GetClass()->GetName());
+	if (InState.ValueType)
+	{
+		PropertyElement.append_attribute("ValueType") = STRING_CAST_UE_TO_PUGI(*InState.ValueType->GetName());
+	}
+	pugi::xml_attribute ValueAttribute = PropertyElement.append_attribute("Value");
+
+	// booleans
+	if (InState.FieldType == FBoolProperty::StaticClass())
+	{
+		ValueAttribute.set_value(CastFieldChecked<FBoolProperty>(InState.ValueProperty)->GetPropertyValue_InContainer(InState.ValueData, InArrayIndex));
+	}
+	
+	// unsigned bytes & enumerations
+	else if (InState.FieldType == FEnumProperty::StaticClass())
+	{
+		FEnumProperty* EnumProperty = CastFieldChecked<FEnumProperty>(InState.ValueProperty);
+
+		ValueAttribute.set_value(STRING_CAST_UE_TO_PUGI(*EnumProperty->GetEnum()->GetNameStringByValue(EnumProperty->GetUnderlyingProperty()->GetSignedIntPropertyValue(EnumProperty->ContainerPtrToValuePtr<void>(InState.ValueData, InArrayIndex)))));
+	}
+	else if (InState.FieldType == FByteProperty::StaticClass())
+	{
+		FByteProperty* ByteProperty = CastFieldChecked<FByteProperty>(InState.ValueProperty);
+
+		if (ByteProperty->IsEnum())
+		{
+			ValueAttribute.set_value(STRING_CAST_UE_TO_PUGI(*ByteProperty->Enum->GetNameStringByValue(ByteProperty->GetPropertyValue_InContainer(InState.ValueData, InArrayIndex))));
+		}
+		else
+		{
+			ValueAttribute.set_value((int)ByteProperty->GetPropertyValue_InContainer(InState.ValueData, InArrayIndex));
+		}
+	}
+	
+	// floating point numbers
+	else if (InState.FieldType == FDoubleProperty::StaticClass())
+	{
+		ValueAttribute.set_value(CastFieldChecked<FDoubleProperty>(InState.ValueProperty)->GetPropertyValue_InContainer(InState.ValueData, InArrayIndex));
+	}
+	else if (InState.FieldType == FFloatProperty::StaticClass())
+	{
+		ValueAttribute.set_value(CastFieldChecked<FFloatProperty>(InState.ValueProperty)->GetPropertyValue_InContainer(InState.ValueData, InArrayIndex));
+	}
+
+	// signed integers
+	else if (InState.FieldType == FIntProperty::StaticClass())
+	{
+		ValueAttribute.set_value((int)CastFieldChecked<FIntProperty>(InState.ValueProperty)->GetPropertyValue_InContainer(InState.ValueData, InArrayIndex));
+	}
+	else if (InState.FieldType == FInt8Property::StaticClass())
+	{
+		ValueAttribute.set_value((int)CastFieldChecked<FInt8Property>(InState.ValueProperty)->GetPropertyValue_InContainer(InState.ValueData, InArrayIndex));
+	}
+	else if (InState.FieldType == FInt16Property::StaticClass())
+	{
+		ValueAttribute.set_value((int)CastFieldChecked<FInt16Property>(InState.ValueProperty)->GetPropertyValue_InContainer(InState.ValueData, InArrayIndex));
+	}
+	else if (InState.FieldType == FInt64Property::StaticClass())
+	{
+		ValueAttribute.set_value((long long)CastFieldChecked<FInt64Property>(InState.ValueProperty)->GetPropertyValue_InContainer(InState.ValueData, InArrayIndex));
+	}
+	
+	// unsigned integers
+	else if (InState.FieldType == FUInt16Property::StaticClass())
+	{
+		ValueAttribute.set_value((unsigned int)CastFieldChecked<FUInt16Property>(InState.ValueProperty)->GetPropertyValue_InContainer(InState.ValueData, InArrayIndex));
+	}
+	else if (InState.FieldType == FUInt32Property::StaticClass())
+	{
+		ValueAttribute.set_value((unsigned int)CastFieldChecked<FUInt32Property>(InState.ValueProperty)->GetPropertyValue_InContainer(InState.ValueData, InArrayIndex));
+	}
+	else if (InState.FieldType == FUInt64Property::StaticClass())
+	{
+		ValueAttribute.set_value((unsigned long long)CastFieldChecked<FUInt64Property>(InState.ValueProperty)->GetPropertyValue_InContainer(InState.ValueData, InArrayIndex));
+	}
+
+	// names, strings & text
+	else if (InState.FieldType == FNameProperty::StaticClass())
+	{
+		ValueAttribute.set_value(STRING_CAST_UE_TO_PUGI(*CastFieldChecked<FNameProperty>(InState.ValueProperty)->GetPropertyValue_InContainer(InState.ValueData, InArrayIndex).ToString()));
+	}
+	else if (InState.FieldType == FStrProperty::StaticClass())
+	{
+		ValueAttribute.set_value(STRING_CAST_UE_TO_PUGI(*CastFieldChecked<FStrProperty>(InState.ValueProperty)->GetPropertyValue_InContainer(InState.ValueData, InArrayIndex)));
+	}
+	else if (InState.FieldType == FTextProperty::StaticClass())
+	{
+		const FText& TextValue = CastFieldChecked<FTextProperty>(InState.ValueProperty)->GetPropertyValue_InContainer(InState.ValueData, InArrayIndex);
+		if (EnumHasAnyFlags(Flags, EStructSerializerBackendFlags::WriteTextAsComplexString))
+		{
+			FString TextValueString;
+			FTextStringHelper::WriteToBuffer(TextValueString, TextValue);
+			ValueAttribute.set_value(STRING_CAST_UE_TO_PUGI(*TextValueString));
+		}
+		else
+		{
+			ValueAttribute.set_value(STRING_CAST_UE_TO_PUGI(*TextValue.ToString()));
+		}
+	}
+
+	// classes & objects
+	else if (InState.FieldType == FSoftClassProperty::StaticClass())
+	{
+		FSoftObjectPtr const& Value = CastFieldChecked<FSoftClassProperty>(InState.ValueProperty)->GetPropertyValue_InContainer(InState.ValueData, InArrayIndex);
+		ValueAttribute.set_value(STRING_CAST_UE_TO_PUGI(*(Value.IsValid() ? Value->GetPathName() : FString())));
+	}
+	else if (InState.FieldType == FWeakObjectProperty::StaticClass())
+	{
+		FWeakObjectPtr const& Value = CastFieldChecked<FWeakObjectProperty>(InState.ValueProperty)->GetPropertyValue_InContainer(InState.ValueData, InArrayIndex);
+		ValueAttribute.set_value(STRING_CAST_UE_TO_PUGI(*(Value.IsValid() ? Value.Get()->GetPathName() : FString())));
+	}
+	else if (InState.FieldType == FSoftObjectProperty::StaticClass())
+	{
+		FSoftObjectPtr const& Value = CastFieldChecked<FSoftObjectProperty>(InState.ValueProperty)->GetPropertyValue_InContainer(InState.ValueData, InArrayIndex);
+		ValueAttribute.set_value(STRING_CAST_UE_TO_PUGI(*Value.ToString()));
+	}
+	else if (FObjectProperty* ObjectProperty = CastField<FObjectProperty>(InState.ValueProperty))
+	{
+		// @TODO: Could this be expanded to include everything derived from FObjectPropertyBase?
+		// Generic handling for a property type derived from FObjectProperty that is obtainable as a pointer and will be stored using its path.
+		// This must come after all the more specialized handlers for object property types.
+		UObject* const Value = ObjectProperty->GetObjectPropertyValue_InContainer(InState.ValueData, InArrayIndex);
+		ValueAttribute.set_value(STRING_CAST_UE_TO_PUGI(*(Value ? Value->GetPathName() : FString())));
+	}
+	// unsupported property type
+	else
+	{
+		const FString ValueType = InState.ValueType != nullptr ? InState.ValueType->GetFName().ToString() : TEXT("unknown");
+		UE_LOG(LogXmlSerialization, Verbose, TEXT("FXmlStructSerializerBackend: Property %s cannot be serialized, because its type (%s) is not supported")
+			, *InState.ValueProperty->GetFName().ToString()
+			, *ValueType);
+	}
+}
diff --git a/Engine/Plugins/Experimental/SerializationUtils/Source/XmlSerialization/Private/Formatters/XmlArchiveInputFormatter.cpp b/Engine/Plugins/Experimental/SerializationUtils/Source/XmlSerialization/Private/Formatters/XmlArchiveInputFormatter.cpp
new file mode 100644
--- /dev/null
+++ b/Engine/Plugins/Experimental/SerializationUtils/Source/XmlSerialization/Private/Formatters/XmlArchiveInputFormatter.cpp
@@ -0,0 +1,936 @@
+// Copyright Epic Games, Inc. All Rights Reserved.
+
+#include "Formatters/XmlArchiveInputFormatter.h"
+#include "Misc/Base64.h"
+#include "Misc/SecureHash.h"
+#include "UObject/SoftObjectPath.h"
+#include "UObject/SoftObjectPtr.h"
+#include "UObject/WeakObjectPtr.h"
+#include "UObject/LazyObjectPtr.h"
+#include "UObject/ObjectPtr.h"
+
+#include "Logging/LogMacros.h"
+#include "Utils/XmlWriter.h"
+#include "Utils/XmlUtils.h"
+#include "XmlFormatterHelper.h"
+#include "XmlSerializationModule.h"
+#include "UObject/UObjectIterator.h"
+
+#if WITH_TEXT_ARCHIVE_SUPPORT
+
+// Private Implementation
+class FXmlArchiveInputFormatter::FImpl
+{
+public:
+	FImpl()
+	{
+		XmlElements.Reserve(64);
+	}
+	
+	// Copy constructor is used when making a sub-tree. Not Tested.
+	FImpl(const FImpl& InOther)
+	{
+		XmlDocument.reset(InOther.XmlDocument);
+		XmlParseResult = InOther.XmlParseResult;
+		const std::string NodePath = InOther.XmlElements.Top().path();
+		XmlElements.Push(XmlDocument.first_element_by_path(NodePath.c_str()));
+	}
+
+	bool Pop(const TCHAR* InFuncName, const TCHAR* InTagName)
+	{
+		using namespace UE::XmlSerialization::Private;
+		
+		// Check if the tag name matches.
+		if (!Utils::IsElementA(XmlElements.Top(), InTagName))
+		{
+			UE_LOG(LogXmlSerialization, Warning,
+				TEXT("FXmlArchiveInputFormatter::%s: Requested to leave \"%s\", but it is \"%s\"."),
+				InFuncName, InTagName, *Utils::GetInfo(XmlElements.Top()));
+		}
+		
+		XmlElements.Pop();
+		return true;
+	}
+
+	void Push(const pugi::xml_node& InXmlNode)
+	{
+		XmlElements.Push(InXmlNode);
+	}
+
+	pugi::xml_attribute GetValueAttrib() const
+	{
+		using namespace UE::XmlSerialization::Private;
+		const pugi::xml_attribute Attrib = !XmlAttrib.empty() ? XmlAttrib : XmlElement().attribute("Value");
+		if (Attrib.empty())
+		{
+			UE_LOG(LogXmlSerialization, Error,
+				TEXT("FXmlArchiveInputFormatter::Serialize: No current attribute. Current Node: \"%s\""),
+				*Utils::GetInfo(XmlElement()));
+		}
+		return Attrib;
+	}
+
+	/** Current xml element being traversed. */
+	const pugi::xml_node& XmlElement() const { return XmlElements.Top(); }
+
+	pugi::xml_node& XmlElementMutable() { return XmlElements.Top(); }
+	
+	/** Current Document. */
+	pugi::xml_document XmlDocument;
+
+	/** Parse result of parsing the document. */
+	pugi::xml_parse_result XmlParseResult;
+
+	/**
+	 * Stack of xml element being traversed.
+	 * Using a stack so we can push empty nodes and pop them back up to the previous existing parent.
+	 * This allows traversing missing tags and recover.
+	 */
+	TArray<pugi::xml_node> XmlElements;
+
+	/** Current Attribute being traversed. */
+	pugi::xml_attribute XmlAttrib;
+
+	/** Last Array/Stream/Map element that was traversed. */
+	pugi::xml_node LastArrayElement;
+};
+
+FXmlArchiveInputFormatter::FXmlArchiveInputFormatter(FArchive& InInner, UObject* InRootObject, const TFunction<UObject* (const FPackageIndex)>& InResolveObject) 
+	: Impl(new FImpl)
+	, Inner(InInner)
+	, ResolveObject(InResolveObject)
+	, RootObject(InRootObject)
+{
+	Inner.SetIsTextFormat(true);
+	Inner.ArAllowLazyLoading = false;
+
+	{
+		// Read the whole thing in memory and parse it.
+		TArray<int8> Buffer;
+		Buffer.SetNum(Inner.TotalSize());
+		Inner.Serialize(Buffer.GetData(), Buffer.Num());
+		Impl->XmlParseResult = Impl->XmlDocument.load_buffer(Buffer.GetData(), Buffer.Num(),
+		                                                              pugi::parse_default | pugi::parse_declaration |
+		                                                              pugi::parse_comments);
+		if (Impl->XmlParseResult.status != pugi::status_ok)
+		{
+			//Impl->XmlParseResult.description() -> convert to FString.
+			UE_LOG(LogXmlSerialization, Error, TEXT("FXmlArchiveInputFormatter: Failed to parse xml document.")); 
+		}
+		else
+		{
+			// Find the root of the document.
+			pugi::xml_node RootNode = Impl->XmlDocument.document_element();
+			if (FCStringAnsi::Strcmp(RootNode.name(), "ObjectFileDocument") != 0)
+			{
+				RootNode = RootNode.next_sibling("ObjectFileDocument");
+				if (RootNode.empty())
+				{
+					RootNode = Impl->XmlDocument.document_element().child("ObjectFileDocument");
+				}
+			}
+			
+			if (RootNode.empty())
+			{
+				UE_LOG(LogXmlSerialization, Error, TEXT("FXmlArchiveInputFormatter: Document is missing the root tag.")); 
+			}
+			
+			Impl->Push(RootNode);
+		}
+	}
+}
+
+FXmlArchiveInputFormatter::FXmlArchiveInputFormatter(const FXmlArchiveInputFormatter& InOther)
+	: Impl(new FImpl(*InOther.Impl))
+	, Inner(InOther.Inner)
+	, ResolveObject(InOther.ResolveObject)
+	, RootObject(InOther.RootObject)
+{
+	
+}
+
+FXmlArchiveInputFormatter::~FXmlArchiveInputFormatter()
+{
+}
+
+bool FXmlArchiveInputFormatter::IsParseResultStatusOk() const
+{
+	return (Impl->XmlParseResult.status == pugi::status_ok);
+}
+
+FArchive& FXmlArchiveInputFormatter::GetUnderlyingArchive()
+{
+	return Inner;
+}
+
+FStructuredArchiveFormatter* FXmlArchiveInputFormatter::CreateSubtreeReader()
+{
+	FXmlArchiveInputFormatter* Cloned = new FXmlArchiveInputFormatter(*this);
+	// TODO: not sure this works.
+	Cloned->Impl->XmlDocument.append_copy(Impl->XmlElement());
+	Cloned->Impl->Push(Cloned->Impl->XmlDocument.document_element());
+	return Cloned;
+}
+
+bool FXmlArchiveInputFormatter::HasDocumentTree() const
+{
+	return true;
+}
+
+void FXmlArchiveInputFormatter::EnterRecord()
+{
+	using namespace UE::XmlSerialization::Private;
+	const pugi::xml_node RecordElement = Impl->XmlElement().child("Record");
+	if (RecordElement.empty())
+	{
+		UE_LOG(LogXmlSerialization, Warning,
+		       TEXT("FXmlArchiveInputFormatter::EnterRecord: Node \"%s\" has no \"Record\" child."),
+		       		*Utils::GetInfo(Impl->XmlElement()));
+	}
+	
+	Impl->Push(RecordElement);
+}
+
+void FXmlArchiveInputFormatter::LeaveRecord()
+{
+	Impl->Pop(TEXT("LeaveRecord"), TEXT("Record"));
+}
+
+void FXmlArchiveInputFormatter::EnterField(FArchiveFieldName InName)
+{
+	using namespace UE::XmlSerialization::Private;
+	// We are looking for a "Field" node with the attribute "Name" being the given name.
+	// Note: this only searches by attribute name/value.
+	const pugi::xml_node FieldElement = Impl->XmlElement().find_child_by_attribute("Field", "Name", STRING_CAST_UE_TO_PUGI(InName.Name));
+	if (FieldElement.empty())
+	{
+		UE_LOG(LogXmlSerialization, Warning,
+		       TEXT("FXmlArchiveInputFormatter::EnterField: Failed to find Field \"%s\". Current Node: \"%s\"."), InName.Name,
+		       *Utils::GetInfo(Impl->XmlElement()));
+	}
+	
+	Impl->Push(FieldElement);
+}
+
+void FXmlArchiveInputFormatter::LeaveField()
+{
+	Impl->Pop(TEXT("LeaveField"), TEXT("Field"));
+}
+
+bool FXmlArchiveInputFormatter::TryEnterField(FArchiveFieldName InName, bool bInEnterWhenSaving)
+{
+	const pugi::xml_node FieldElement = Impl->XmlElement().find_child_by_attribute("Field", "Name", STRING_CAST_UE_TO_PUGI(InName.Name));
+	if (FieldElement.empty())
+	{
+		return false; 
+	}
+	
+	Impl->Push(FieldElement);
+	return true;
+}
+
+void FXmlArchiveInputFormatter::EnterArray(int32& OutNumElements)
+{
+	using namespace UE::XmlSerialization::Private;
+	const pugi::xml_node ArrayElement = Impl->XmlElement().child("Array");
+	
+	if (!ArrayElement.empty())
+	{
+		const pugi::xml_attribute SizeAttrib = ArrayElement.attribute("Size");
+		if (SizeAttrib.empty())
+		{
+			UE_LOG(LogXmlSerialization, Error,
+				TEXT("FXmlArchiveInputFormatter::EnterArray: Node \"%s\" doesn't have a \"Size\" attribute."),
+				*Utils::GetInfo(Impl->XmlElement())); 
+		}
+	
+		OutNumElements = !SizeAttrib.empty() ? SizeAttrib.as_int() : 0;
+	}
+	else
+	{
+		OutNumElements = 0;
+		
+		UE_LOG(LogXmlSerialization, Error,
+			TEXT("FXmlArchiveInputFormatter::EnterArray: Node \"%s\" doesn't have an \"Array\" child."),
+			*Utils::GetInfo(Impl->XmlElement())); 
+	}
+	
+	Impl->LastArrayElement = pugi::xml_node();
+	Impl->Push(ArrayElement);
+}
+
+void FXmlArchiveInputFormatter::LeaveArray()
+{
+	Impl->Pop(TEXT("LeaveArray"), TEXT("Array"));
+	Impl->LastArrayElement = pugi::xml_node();
+}
+
+void FXmlArchiveInputFormatter::EnterArrayElement()
+{
+	using namespace UE::XmlSerialization::Private;
+	pugi::xml_node ArrayElementElement;
+	
+	if (!Impl->LastArrayElement.empty())
+	{
+		ArrayElementElement = Impl->LastArrayElement.next_sibling("ArrayElement");
+		
+		if (ArrayElementElement.empty())
+		{
+			UE_LOG(LogXmlSerialization, Error,
+				TEXT("FXmlArchiveInputFormatter::EnterArrayElement: Node \"%s\" doesn't have a \"ArrayElement\" sibling."),
+				*Utils::GetInfo(Impl->LastArrayElement)); 
+		}
+	}
+	else
+	{
+		ArrayElementElement = Impl->XmlElement().child("ArrayElement");
+		
+		if (ArrayElementElement.empty())
+		{
+			UE_LOG(LogXmlSerialization, Error,
+				TEXT("FXmlArchiveInputFormatter::EnterArrayElement: Node \"%s\" doesn't have a \"ArrayElement\" child."),
+				*Utils::GetInfo(Impl->XmlElement())); 
+		}
+	}
+
+	Impl->Push(ArrayElementElement);
+}
+
+void FXmlArchiveInputFormatter::LeaveArrayElement()
+{
+	using namespace UE::XmlSerialization::Private;
+	Impl->LastArrayElement =  Utils::IsElementA(Impl->XmlElement(),TEXT("ArrayElement")) ? Impl->XmlElement() : pugi::xml_node();
+	Impl->Pop(TEXT("LeaveArrayElement"), TEXT("ArrayElement"));
+}
+
+void FXmlArchiveInputFormatter::EnterStream()
+{
+	using namespace UE::XmlSerialization::Private;
+	const pugi::xml_node StreamElement = Impl->XmlElement().child("Stream");
+
+	if (StreamElement.empty())
+	{
+		UE_LOG(LogXmlSerialization, Error,
+			TEXT("FXmlArchiveInputFormatter::EnterStream: Node \"%s\" doesn't have a \"Stream\" child."),
+			*Utils::GetInfo(Impl->XmlElement())); 
+	}
+	
+	Impl->Push(StreamElement);
+	Impl->LastArrayElement = pugi::xml_node();
+}
+
+void FXmlArchiveInputFormatter::LeaveStream()
+{
+	Impl->Pop(TEXT("LeaveStream"), TEXT("Stream"));
+	Impl->LastArrayElement = pugi::xml_node();
+}
+
+void FXmlArchiveInputFormatter::EnterStreamElement()
+{
+	using namespace UE::XmlSerialization::Private;
+	pugi::xml_node StreamElementElement;
+	
+	if (!Impl->LastArrayElement.empty())
+	{
+		StreamElementElement = Impl->LastArrayElement.next_sibling("StreamElement");
+
+		if (!StreamElementElement.empty())
+		{
+			UE_LOG(LogXmlSerialization, Error,
+				TEXT("FXmlArchiveInputFormatter::EnterStreamElement: Node \"%s\" doesn't have a \"StreamElement\" sibling."),
+				*Utils::GetInfo(Impl->LastArrayElement));
+		}
+	}
+	else
+	{
+		StreamElementElement = Impl->XmlElement().child("StreamElement");
+
+		if (!StreamElementElement.empty())
+		{
+			UE_LOG(LogXmlSerialization, Error,
+				TEXT("FXmlArchiveInputFormatter::EnterStreamElement: Node \"%s\" doesn't have a \"StreamElement\" child."),
+				*Utils::GetInfo(Impl->XmlElement()));
+		}
+	}
+	
+	Impl->Push(StreamElementElement);
+}
+
+void FXmlArchiveInputFormatter::LeaveStreamElement()
+{
+	using namespace UE::XmlSerialization::Private;
+	Impl->LastArrayElement =  Utils::IsElementA(Impl->XmlElement(),TEXT("StreamElement")) ? Impl->XmlElement() : pugi::xml_node();
+	Impl->Pop(TEXT("LeaveStreamElement"), TEXT("StreamElement"));
+}
+
+void FXmlArchiveInputFormatter::EnterMap(int32& OutNumElements)
+{
+	using namespace UE::XmlSerialization::Private;
+	const pugi::xml_node MapElement = Impl->XmlElement().child("Map");
+	if (!MapElement.empty())
+	{
+		const pugi::xml_attribute SizeAttrib = MapElement.attribute("Size");
+
+		if (SizeAttrib.empty())
+		{
+			UE_LOG(LogXmlSerialization, Error, TEXT("FXmlArchiveInputFormatter::EnterMap: Node \"%hs\" doesn't have a \"Size\" attribute."), MapElement.name()); 
+		}
+
+		OutNumElements = !SizeAttrib.empty() ? SizeAttrib.as_int() : 0;
+		
+		Impl->Push(MapElement);
+		Impl->LastArrayElement = pugi::xml_node();
+		return;
+	}
+	
+	// Fallback mode: There might be a bug on how MapProperties (or Struct?) are serialized. The code is different for loading vs saving.
+	// See: UStruct::SerializeVersionedTaggedProperties. When loading it wants to enter a map, while when saving, it enters a record.
+	pugi::xml_node RecordElement = Impl->XmlElement().child("Record");
+	if (!RecordElement.empty())
+	{
+		// Remember we where supposed to get into a Map.
+		RecordElement.append_attribute("__HACK_SHOULD_BE__") = "Map"; 
+
+		// The number of elements are going to be the number of Field nodes.
+		// It will then try to enter MapElements, but we'll be looking for Field Nodes instead.
+		OutNumElements = 0;
+		for (pugi::xml_node Field : RecordElement.children("Field"))
+		{
+			OutNumElements++;
+		}	
+
+		UE_LOG(LogXmlSerialization, Verbose, TEXT("FXmlArchiveInputFormatter:EnterMap: Fallback: Entering \"Record\" node instead, with %d entries (Fields)."), OutNumElements);
+		Impl->Push(RecordElement);
+		Impl->LastArrayElement = pugi::xml_node();
+		return;
+	}
+
+	UE_LOG(LogXmlSerialization, Error,
+		TEXT("FXmlArchiveInputFormatter::EnterMap: Node \"%s\" doesn't have a \"Map\" (nor \"Record\") child."),
+		*Utils::GetInfo(Impl->XmlElement()));
+	
+	OutNumElements = 0;
+	Impl->Push(pugi::xml_node());
+	Impl->LastArrayElement = pugi::xml_node();
+}
+
+inline bool ElementIsARecordButShouldBeAMap(const pugi::xml_node& InElement)
+{
+	if (FCStringAnsi::Strcmp(InElement.name(), "Record") == 0)
+	{
+		const pugi::xml_attribute HackAttrib = InElement.attribute("__HACK_SHOULD_BE__");
+		return (!HackAttrib.empty() && FCStringAnsi::Strcmp(HackAttrib.as_string(), "Map") == 0) ? true : false;
+	}
+	return false;
+}
+
+void FXmlArchiveInputFormatter::LeaveMap()
+{
+	if (ElementIsARecordButShouldBeAMap(Impl->XmlElement()))
+	{
+		Impl->Pop(TEXT("LeaveMap"), TEXT("Record"));
+	}
+	else
+	{
+		Impl->Pop(TEXT("LeaveMap"), TEXT("Map"));
+	}
+	Impl->LastArrayElement = pugi::xml_node();
+}
+
+void FXmlArchiveInputFormatter::EnterMapElement(FString& OutName)
+{
+	using namespace UE::XmlSerialization::Private;
+	
+	// Fallback code for the Map vs Record bug.
+	if (ElementIsARecordButShouldBeAMap(Impl->XmlElement()))
+	{
+		pugi::xml_node FieldElement;
+		
+		if (!Impl->LastArrayElement.empty())
+		{
+			FieldElement = Impl->LastArrayElement.next_sibling("Field");
+			
+			if (FieldElement.empty())
+			{
+				UE_LOG(LogXmlSerialization, Error,
+				   TEXT("FXmlArchiveInputFormatter::EnterMapElement:(Fallback) \"%s\" node doesn't have a \"Field\" sibling, which was supposed to be a \"MapElement\"."),
+					*Utils::GetInfo(Impl->LastArrayElement));
+			}
+		}
+		else
+		{
+			FieldElement = Impl->XmlElement().child("Field");
+
+			if (FieldElement.empty())
+			{
+				UE_LOG(LogXmlSerialization, Error,
+				   TEXT("FXmlArchiveInputFormatter::EnterMapElement:(Fallback) \"%s\" node doesn't have a \"Field\" child, which was supposed to be a \"MapElement\"."),
+					*Utils::GetInfo(Impl->XmlElement()));
+			}
+		}
+
+		if (!FieldElement.empty())
+		{
+			const pugi::xml_attribute NameAttrib = FieldElement.attribute("Name");
+			OutName = NameAttrib.as_string();
+			if (NameAttrib.empty())
+			{
+				UE_LOG(LogXmlSerialization, Error,
+					TEXT("FXmlArchiveInputFormatter::EnterMapElement:(Fallback) \"%s\" node doesn't have a \"Name\" attribute."),
+					*Utils::GetInfo(FieldElement));	
+			}
+
+			// Remember we where supposed to get into a MapElement.
+			FieldElement.append_attribute("__HACK_SHOULD_BE__") = "MapElement";
+		}
+
+		Impl->Push(FieldElement);
+		return; // fallback job done.
+	}
+	
+	pugi::xml_node MapElementElement;
+	
+	if (!Impl->LastArrayElement.empty())
+	{
+		MapElementElement = Impl->LastArrayElement.next_sibling("MapElement");
+
+		if (MapElementElement.empty())
+		{
+			UE_LOG(LogXmlSerialization, Error,
+				TEXT("FXmlArchiveInputFormatter::EnterMapElement: \"%s\" node doesn't have a \"MapElement\" sibling."),
+				*Utils::GetInfo(Impl->LastArrayElement)); 
+		}
+	}
+	else
+	{
+		MapElementElement = Impl->XmlElement().child("MapElement");
+
+		if (MapElementElement.empty())
+		{
+			UE_LOG(LogXmlSerialization, Error,
+				TEXT("FXmlArchiveInputFormatter::EnterMapElement: \"%s\" node doesn't have a \"MapElement\" child."),
+				*Utils::GetInfo(Impl->XmlElement())); 
+		}		
+	}
+	
+	if (!MapElementElement.empty())
+	{
+		const pugi::xml_attribute NameAttrib = MapElementElement.attribute("Name");
+		OutName = NameAttrib.as_string();
+		if (NameAttrib.empty())
+		{
+			UE_LOG(LogXmlSerialization, Error,
+				TEXT("FXmlArchiveInputFormatter::EnterMapElement: \"%s\" node doesn't have a \"Name\" attribute."),
+				*Utils::GetInfo(MapElementElement));	
+		}
+	}
+	
+	Impl->Push(MapElementElement);
+}
+
+void FXmlArchiveInputFormatter::LeaveMapElement()
+{
+	using namespace UE::XmlSerialization::Private;
+	
+	// Check if our current element should be something else (fallback code above).
+	const pugi::xml_attribute HackAttrib = Impl->XmlElement().attribute("__HACK_SHOULD_BE__");
+	if (!HackAttrib.empty())
+	{
+		// It is possible to be in a value node. If so pop it first.
+		if (Utils::IsElementA(Impl->XmlElement(), TEXT("Value")))
+		{
+			Impl->Pop(TEXT("LeaveMapElement"), TEXT("Value"));
+		}
+
+		// Should be in a Field node. If so, keep track of the last element.
+		Impl->LastArrayElement =  Utils::IsElementA(Impl->XmlElement(),TEXT("Field")) ? Impl->XmlElement() : pugi::xml_node();
+		Impl->Pop(TEXT("LeaveMapElement"), TEXT("Field"));
+	}
+	else
+	{
+		Impl->LastArrayElement =  Utils::IsElementA(Impl->XmlElement(),TEXT("MapElement")) ? Impl->XmlElement() : pugi::xml_node();
+		Impl->Pop(TEXT("LeaveMapElement"), TEXT("MapElement"));
+	}
+}
+
+void FXmlArchiveInputFormatter::EnterAttributedValue()
+{
+	using namespace UE::XmlSerialization::Private;
+	const pugi::xml_node ValueElement = Impl->XmlElement().child("AttributedValue");
+	if (ValueElement.empty())
+	{
+		UE_LOG(LogXmlSerialization, Error,
+			TEXT("FXmlArchiveInputFormatter::EnterAttributedValue: Node \"%s\" doesn't have an \"AttributedValue\" child."),
+			*Utils::GetInfo(Impl->XmlElement())); 
+	}
+	Impl->Push(ValueElement);
+}
+
+void FXmlArchiveInputFormatter::EnterAttribute(FArchiveFieldName AttributeName)
+{
+	using namespace UE::XmlSerialization::Private;
+	const pugi::xml_node AttributeElement = Impl->XmlElement().find_child_by_attribute("Attribute", "Name", STRING_CAST_UE_TO_PUGI(AttributeName.Name));
+	if (AttributeElement.empty())
+	{
+		UE_LOG(LogXmlSerialization, Error,
+			TEXT("FXmlArchiveInputFormatter::EnterAttribute: Node \"%s\" doesn't have an \"Attribute\" child of name \"%s\"."),
+			*Utils::GetInfo(Impl->XmlElement()), AttributeName.Name);
+	}
+	
+	Impl->Push(AttributeElement);
+}
+
+void FXmlArchiveInputFormatter::EnterAttributedValueValue()
+{
+	using namespace UE::XmlSerialization::Private;
+	const pugi::xml_node ValueElement = Impl->XmlElement().child("AttributedValueValue");
+	if (ValueElement.empty())
+	{
+		UE_LOG(LogXmlSerialization, Error,
+			TEXT("FXmlArchiveInputFormatter::EnterAttributedValueValue: Node \"%s\" doesn't have an \"AttributedValueValue\" child."),
+			*Utils::GetInfo(Impl->XmlElement())); 
+	}
+	Impl->Push(ValueElement);
+}
+
+bool FXmlArchiveInputFormatter::TryEnterAttributedValueValue()
+{
+	const pugi::xml_node ValueElement = Impl->XmlElement().child("AttributedValueValue");
+	if (!ValueElement.empty())
+	{
+		Impl->Push(ValueElement);
+		return true;
+	}
+	return false;
+}
+
+void FXmlArchiveInputFormatter::LeaveAttribute()
+{
+	using namespace UE::XmlSerialization::Private;
+	// AttributedValueValue and Attribute are at the same level under AttributedValue,
+	// We can pop them if requested.
+	if (Utils::IsElementA(Impl->XmlElement(), TEXT("AttributedValueValue")))
+	{
+		Impl->Pop(TEXT("LeaveAttribute"), TEXT("AttributedValueValue"));
+	}
+	else
+	{
+		Impl->Pop(TEXT("LeaveAttribute"), TEXT("Attribute"));
+	}
+}
+
+void FXmlArchiveInputFormatter::LeaveAttributedValue()
+{
+	using namespace UE::XmlSerialization::Private;
+	// AttributedValueValue and Attribute seem to be used interchangeably, so we
+	// need to pop out of them if popping out of AttributedValue.
+	if (Utils::IsElementA(Impl->XmlElement(), TEXT("AttributedValueValue")))
+	{
+		Impl->Pop(TEXT("LeaveAttributedValue"), TEXT("AttributedValueValue"));
+	}
+	else if (Utils::IsElementA(Impl->XmlElement(), TEXT("Attribute")))
+	{
+		Impl->Pop(TEXT("LeaveAttributedValue"), TEXT("Attribute"));
+	}
+	
+	Impl->Pop(TEXT("LeaveAttributedValue"), TEXT("AttributedValue"));
+}
+
+bool FXmlArchiveInputFormatter::TryEnterAttribute(FArchiveFieldName InAttributeName, bool bInEnterWhenSaving)
+{
+	const pugi::xml_node AttributeElement = Impl->XmlElement().find_child_by_attribute("Attribute", "Name", STRING_CAST_UE_TO_PUGI(InAttributeName.Name));
+	if (!AttributeElement.empty())
+	{
+		Impl->Push(AttributeElement);
+		return true;
+	}
+	return false;
+}
+
+void FXmlArchiveInputFormatter::Serialize(uint8& Value)
+{
+	const pugi::xml_attribute Attrib = Impl->GetValueAttrib();
+	if (!Attrib.empty())
+	{
+		Value = Attrib.as_uint();
+	}
+}
+
+void FXmlArchiveInputFormatter::Serialize(uint16& Value)
+{
+	const pugi::xml_attribute Attrib = Impl->GetValueAttrib();
+	if (!Attrib.empty())
+	{
+		Value = Attrib.as_uint();
+	}
+}
+
+void FXmlArchiveInputFormatter::Serialize(uint32& Value)
+{
+	const pugi::xml_attribute Attrib = Impl->GetValueAttrib();
+	if (!Attrib.empty())
+	{
+		Value = Attrib.as_uint();
+	}
+}
+
+void FXmlArchiveInputFormatter::Serialize(uint64& Value)
+{
+	const pugi::xml_attribute Attrib = Impl->GetValueAttrib();
+	if (!Attrib.empty())
+	{
+		Value = Attrib.as_ullong();
+	}
+}
+
+void FXmlArchiveInputFormatter::Serialize(int8& Value)
+{
+	const pugi::xml_attribute Attrib = Impl->GetValueAttrib();
+	if (!Attrib.empty())
+	{
+		Value = Attrib.as_int();
+	}
+}
+
+void FXmlArchiveInputFormatter::Serialize(int16& Value)
+{
+	const pugi::xml_attribute Attrib = Impl->GetValueAttrib();
+	if (!Attrib.empty())
+	{
+		Value = Attrib.as_int();
+	}
+}
+
+void FXmlArchiveInputFormatter::Serialize(int32& Value)
+{
+	const pugi::xml_attribute Attrib = Impl->GetValueAttrib();
+	if (!Attrib.empty())
+	{
+		Value = Attrib.as_int();
+	}
+}
+
+void FXmlArchiveInputFormatter::Serialize(int64& Value)
+{
+	const pugi::xml_attribute Attrib = Impl->GetValueAttrib();
+	if (!Attrib.empty())
+	{
+		Value = Attrib.as_llong();
+	}
+}
+
+void FXmlArchiveInputFormatter::Serialize(float& Value)
+{
+	const pugi::xml_attribute Attrib = Impl->GetValueAttrib();
+	if (!Attrib.empty())
+	{
+		Value = Attrib.as_float();
+	}
+}
+	
+void FXmlArchiveInputFormatter::Serialize(double& Value)
+{
+	const pugi::xml_attribute Attrib = Impl->GetValueAttrib();
+	if (!Attrib.empty())
+	{
+		Value = Attrib.as_double();
+	}
+}
+
+void FXmlArchiveInputFormatter::Serialize(bool& Value)
+{
+	const pugi::xml_attribute Attrib = Impl->GetValueAttrib();
+	if (!Attrib.empty())
+	{
+		Value = Attrib.as_bool();
+	}
+}
+
+void FXmlArchiveInputFormatter::Serialize(FString& Value)
+{
+	const pugi::xml_attribute Attrib = Impl->GetValueAttrib();
+	if (!Attrib.empty())
+	{
+		Value = Attrib.as_string();	
+		//Value.RemoveFromStart(TEXT("String:"));	// We removed that from serializer. Just a reminder here.
+	}
+}
+
+void FXmlArchiveInputFormatter::Serialize(FName& Value)
+{
+	const pugi::xml_attribute Attrib = Impl->GetValueAttrib();
+	if (!Attrib.empty())
+	{
+		const FString StringValue = Attrib.as_string();
+		Value = FName(*StringValue);
+	}
+}
+
+void FXmlArchiveInputFormatter::Serialize(UObject*& Value)
+{
+	// Special case for nested objects.
+	if (IsNestedObject())
+	{
+		Value = LoadNestedObject();
+	}
+	else
+	{
+		const pugi::xml_attribute Attrib = Impl->GetValueAttrib();
+		const FString StringValue = !Attrib.empty() ? Attrib.as_string() : FString();
+		if (!StringValue.IsEmpty())
+		{
+			FPackageIndex ObjectIndex;
+			LexFromString(ObjectIndex, *StringValue);
+			Value = ResolveObject(ObjectIndex);
+		}
+		else
+		{
+			Value = nullptr;
+		}
+	}
+}
+
+void FXmlArchiveInputFormatter::Serialize(FText& Value)
+{
+	UE_LOG(LogXmlSerialization, Error, TEXT("FXmlArchiveInputFormatter::Serialize: FText value is not supported by this formatter.")); 
+
+	// Unlike FJsonArchiveInputFormatter, my class is not friends with FText. :(
+	
+	// FStructuredArchive ChildArchive(*this);
+	// FText::SerializeText(ChildArchive.Open(), Value);	---> is Private
+	// ChildArchive.Close();
+}
+
+void FXmlArchiveInputFormatter::Serialize(FWeakObjectPtr& Value)
+{
+	UObject* Object = nullptr;
+	Serialize(Object);
+	Value = Object;
+}
+
+void FXmlArchiveInputFormatter::Serialize(FSoftObjectPtr& Value)
+{
+	// Special case for nested objects.
+	if (IsNestedObject())
+	{
+		Value = LoadNestedObject();
+	}
+	else
+	{
+		FSoftObjectPath Path;
+		Serialize(Path);
+		Value = Path;
+	}
+}
+
+void FXmlArchiveInputFormatter::Serialize(FSoftObjectPath& Value)
+{
+	const FString StringValue = Impl->GetValueAttrib().as_string();
+	const auto& Prefix = TEXT("Object:");
+	if (StringValue.StartsWith(Prefix))
+	{
+		Value.SetPath(*StringValue + UE_ARRAY_COUNT(Prefix) - 1);
+	}
+	else
+	{
+		Value.Reset();
+	}
+}
+
+void FXmlArchiveInputFormatter::Serialize(FLazyObjectPtr& Value)
+{
+	const FString StringValue = Impl->GetValueAttrib().as_string();
+	const auto& Prefix = TEXT("Lazy:");
+	if (StringValue.StartsWith(Prefix))
+	{
+		FUniqueObjectGuid Guid;
+		Guid.FromString(*StringValue + UE_ARRAY_COUNT(Prefix) - 1);
+		Value = Guid;
+	}
+	else
+	{
+		Value.Reset();
+	}
+}
+
+void FXmlArchiveInputFormatter::Serialize(FObjectPtr& Value)
+{
+	UObject* Object;
+	Serialize(Object);
+	Value = Object;
+}
+
+void FXmlArchiveInputFormatter::Serialize(TArray<uint8>& Data)
+{
+	// Single line string
+	FString RawData = Impl->GetValueAttrib().as_string();
+	if (!RawData.IsEmpty())
+	{
+		ensure(RawData.RemoveFromStart(TEXT("Base64:")));
+		FBase64::Decode(RawData, Data);
+	}
+}
+
+void FXmlArchiveInputFormatter::Serialize(void* Data, uint64 DataSize)
+{
+	TArray<uint8> Buffer;
+	Serialize(Buffer);
+	check(Buffer.Num() == DataSize);
+	memcpy(Data, Buffer.GetData(), DataSize);
+}
+
+bool FXmlArchiveInputFormatter::IsNestedObject() const
+{
+	const pugi::xml_node RecordChild = Impl->XmlElement().child("Record");
+	if (!RecordChild.empty() && RecordChild.attribute("Class").empty() == false)
+	{
+		return true;
+	}
+	return false;
+}
+
+UObject* FXmlArchiveInputFormatter::LoadNestedObject()
+{
+	using namespace UE::XmlSerialization::Private;
+	const pugi::xml_node RecordChild = Impl->XmlElement().child("Record");
+	if (!RecordChild.empty() && RecordChild.attribute("Class").empty() == false)
+	{
+		const FSoftObjectPath ClassPath(RecordChild.attribute("Class").as_string());
+		const TSoftObjectPtr<UClass> ClassPtr(ClassPath);
+		if (const UClass* const ObjectClass = ClassPtr.LoadSynchronous())
+		{
+			const EObjectFlags ObjectFlags = static_cast<EObjectFlags>(RecordChild.attribute("Flags").as_uint());
+			const FString OriginalObjectName(RecordChild.attribute("Name").as_string());
+			// Name clashing prevention: in order to prevent name clashing, we let the new object be named with
+			// a unique name and rename the root field tag.
+			UObject* const NestedObject = NewObject<UObject>(RootObject, ObjectClass, NAME_None, ObjectFlags);
+
+			if (!NestedObject)
+			{
+				return nullptr;
+			}
+
+			{
+				// Rename the root field tag's name attribute to match what the new object's name is.
+				const pugi::xml_node FieldElement = RecordChild.find_child_by_attribute("Field", "Name", STRING_CAST_UE_TO_PUGI(*OriginalObjectName));
+				if (!FieldElement.empty())
+				{
+					FieldElement.attribute("Name") = STRING_CAST_UE_TO_PUGI(*NestedObject->GetName());
+				}
+			}
+			{
+				FStructuredArchive ChildArchive(*this);
+				FStructuredArchiveRecord ObjectRecord = ChildArchive.Open().EnterRecord();
+				FormatterHelper::SerializeObject(ObjectRecord, NestedObject);
+				ChildArchive.Close();
+			}
+
+			NestedObject->PostLoad();
+
+			return NestedObject;
+		}
+	}
+	return nullptr;
+}
+
+#endif
diff --git a/Engine/Plugins/Experimental/SerializationUtils/Source/XmlSerialization/Private/Formatters/XmlArchiveOutputFormatter.cpp b/Engine/Plugins/Experimental/SerializationUtils/Source/XmlSerialization/Private/Formatters/XmlArchiveOutputFormatter.cpp
new file mode 100644
--- /dev/null
+++ b/Engine/Plugins/Experimental/SerializationUtils/Source/XmlSerialization/Private/Formatters/XmlArchiveOutputFormatter.cpp
@@ -0,0 +1,541 @@
+// Copyright Epic Games, Inc. All Rights Reserved.
+
+#include "Formatters/XmlArchiveOutputFormatter.h"
+#include "Serialization/MemoryWriter.h"
+#include "Misc/Base64.h"
+#include "Misc/SecureHash.h"
+#include "UObject/Object.h"
+#include "UObject/WeakObjectPtr.h"
+#include "UObject/SoftObjectPtr.h"
+#include "UObject/SoftObjectPath.h"
+#include "UObject/LazyObjectPtr.h"
+#include "UObject/ObjectPtr.h"
+
+#include "Utils/XmlWriter.h"
+#include "Utils/XmlUtils.h"
+#include "XmlFormatterHelper.h"
+#include "XmlSerializationModule.h"
+
+#if WITH_TEXT_ARCHIVE_SUPPORT
+
+class FXmlArchiveOutputFormatter::FImpl
+{
+public:
+	FImpl()
+	{
+		pugi::xml_node decl = XmlDoc.prepend_child(pugi::node_declaration);
+		decl.append_attribute("version") = "1.0";
+		decl.append_attribute("encoding") = "utf-8";
+
+		XmlElement = XmlDoc.append_child();
+		XmlElement.set_name("ObjectFileDocument");
+		XmlElement.append_attribute(("Version")) = "1.0";
+	}
+
+	// Note: for serialization, we are going to force popping.
+	bool Pop(const TCHAR* FuncName, const TCHAR* TagName, bool bForce = true)
+	{
+		using namespace UE::XmlSerialization::Private;
+		bool bPop = true;
+		// Check if the tag name matches.
+		if (!Utils::IsElementA(XmlElement, TagName))
+		{
+			UE_LOG(LogXmlSerialization, Error,
+				TEXT("FXmlArchiveOutputFormatter::%s: Requested to leave \"%s\", but it is \"%s\"."),
+				FuncName, TagName, *Utils::GetInfo(XmlElement));
+			bPop = false;
+		}
+		if (bPop || bForce)
+		{
+			XmlElement = XmlElement.parent();
+		}
+		return bPop;
+	}
+
+	template <typename ValueType>
+	void SetValue(ValueType Value)
+	{
+		if (!XmlAttrib.empty())
+		{
+			XmlAttrib.set_value(Value);
+		}
+		else
+		{
+			// TODO: Investigate why this is not working.
+			//XmlElement.set_value(*LexToString(Value));
+			XmlElement.append_attribute(("Value")).set_value(Value);
+		}
+	}
+
+	void SetValue(const TCHAR* Value)
+	{
+		if (!XmlAttrib.empty())
+		{
+			XmlAttrib.set_value(STRING_CAST_UE_TO_PUGI(Value));
+		}
+		else
+		{
+			// TODO: Investigate why this is not working.
+			//XmlElement.set_value(Value);
+			XmlElement.append_attribute(("Value")).set_value(STRING_CAST_UE_TO_PUGI(Value));
+		}
+	}
+
+	
+	pugi::xml_document XmlDoc;
+	pugi::xml_node XmlElement;
+	pugi::xml_attribute XmlAttrib;
+	bool bDocumentSaveAttempted = false;
+};
+
+FXmlArchiveOutputFormatter::FXmlArchiveOutputFormatter(FArchive& InInner) 
+	: Impl(new FImpl())
+	, Inner(InInner)
+{
+	Inner.SetIsTextFormat(true);
+}
+
+FXmlArchiveOutputFormatter::~FXmlArchiveOutputFormatter()
+{
+	if (!Impl->bDocumentSaveAttempted)
+	{
+		if (!SaveDocumentToInnerArchive())
+		{
+			UE_LOG(LogXmlSerialization, Error,
+				TEXT("~FXmlArchiveOutputFormatter: Failed to save XML document to inner archive. Archive might be already destroyed, use SaveDocumentToInnerArchive() to ensure proper order."));
+		}
+	}
+	delete Impl;
+}
+
+bool FXmlArchiveOutputFormatter::SaveDocumentToInnerArchive(EXmlSerializationEncoding InEncoding)
+{
+	using namespace UE::XmlSerialization::Private;
+	// We let the user force saving again, but at least issue a warning about it.
+	if (Impl->bDocumentSaveAttempted)
+	{
+		UE_LOG(LogXmlSerialization, Warning,
+			TEXT("FXmlArchiveOutputFormatter: Document has already been saved to archive. Prepare for unforeseen consequences."));
+	}
+	
+	bool bSuccess = true;
+	try
+	{
+		const pugi::xml_encoding encoding = Utils::ToPugiEncoding(InEncoding);
+		Impl->bDocumentSaveAttempted = true;
+		FXmlWriter XmlWriter(&Inner);
+		Impl->XmlDoc.save(XmlWriter, "  ", pugi::format_default, encoding);
+	}
+	catch (...)
+	{
+		bSuccess = false;
+	}
+	return bSuccess;
+}
+
+
+FArchive& FXmlArchiveOutputFormatter::GetUnderlyingArchive()
+{
+	return Inner;
+}
+
+bool FXmlArchiveOutputFormatter::HasDocumentTree() const
+{
+	return true;
+}
+
+void FXmlArchiveOutputFormatter::EnterRecord()
+{
+	pugi::xml_node RecordElement = Impl->XmlElement.append_child(pugi::node_element);
+	RecordElement.set_name("Record");
+	Impl->XmlElement = RecordElement;	
+}
+
+void FXmlArchiveOutputFormatter::LeaveRecord()
+{
+	Impl->Pop(TEXT("LeaveRecord"), TEXT("Record"));
+}
+
+void FXmlArchiveOutputFormatter::EnterField(FArchiveFieldName Name)
+{
+	pugi::xml_node FieldElement = Impl->XmlElement.append_child(pugi::node_element);
+	FieldElement.set_name("Field");
+	FieldElement.append_attribute("Name") = STRING_CAST_UE_TO_PUGI(Name.Name);
+	Impl->XmlElement = FieldElement;
+}
+
+void FXmlArchiveOutputFormatter::LeaveField()
+{
+	Impl->Pop(TEXT("LeaveField"),TEXT("Field"));
+}
+
+bool FXmlArchiveOutputFormatter::TryEnterField(FArchiveFieldName Name, bool bEnterWhenSaving)
+{
+	if (bEnterWhenSaving)
+	{
+		EnterField(Name);
+	}
+	return bEnterWhenSaving;
+}
+
+void FXmlArchiveOutputFormatter::EnterArray(int32& NumElements)
+{
+	pugi::xml_node ArrayElement = Impl->XmlElement.append_child(pugi::node_element);
+	ArrayElement.set_name("Array");
+	ArrayElement.append_attribute("Size").set_value(NumElements);
+	Impl->XmlElement = ArrayElement;
+}
+
+void FXmlArchiveOutputFormatter::LeaveArray()
+{
+	Impl->Pop(TEXT("LeaveArray"),TEXT("Array"));
+}
+
+void FXmlArchiveOutputFormatter::EnterArrayElement()
+{
+	pugi::xml_node ArrayElementElement = Impl->XmlElement.append_child(pugi::node_element);
+	ArrayElementElement.set_name("ArrayElement");
+	Impl->XmlElement = ArrayElementElement;
+}
+
+void FXmlArchiveOutputFormatter::LeaveArrayElement()
+{
+	Impl->Pop(TEXT("LeaveArrayElement"), TEXT("ArrayElement"));
+}
+
+void FXmlArchiveOutputFormatter::EnterStream()
+{
+	pugi::xml_node StreamElement = Impl->XmlElement.append_child(pugi::node_element);
+	StreamElement.set_name("Stream");
+	Impl->XmlElement = StreamElement;
+}
+
+void FXmlArchiveOutputFormatter::LeaveStream()
+{
+	Impl->Pop(TEXT("LeaveStream"), TEXT("Stream"));
+}
+
+void FXmlArchiveOutputFormatter::EnterStreamElement()
+{
+	pugi::xml_node StreamElementElement = Impl->XmlElement.append_child(pugi::node_element);
+	StreamElementElement.set_name("StreamElement");
+	Impl->XmlElement = StreamElementElement;
+}
+
+void FXmlArchiveOutputFormatter::LeaveStreamElement()
+{
+	Impl->Pop(TEXT("LeaveStreamElement"),TEXT("StreamElement"));
+}
+
+void FXmlArchiveOutputFormatter::EnterMap(int32& NumElements)
+{
+	pugi::xml_node MapElement = Impl->XmlElement.append_child(pugi::node_element);
+	MapElement.set_name("Map");
+	MapElement.append_attribute("Size").set_value(NumElements);
+	Impl->XmlElement = MapElement;
+}
+
+void FXmlArchiveOutputFormatter::LeaveMap()
+{
+	Impl->Pop(TEXT("LeaveMap"),TEXT("Map"));
+}
+
+void FXmlArchiveOutputFormatter::EnterMapElement(FString& Name)
+{
+	pugi::xml_node MapElementElement = Impl->XmlElement.append_child(pugi::node_element);
+	MapElementElement.set_name("MapElement");
+	MapElementElement.append_attribute("Name") = STRING_CAST_UE_TO_PUGI(*Name);
+	Impl->XmlElement = MapElementElement;
+}
+
+void FXmlArchiveOutputFormatter::LeaveMapElement()
+{
+	Impl->Pop(TEXT("LeaveMapElement"),TEXT("MapElement"));
+}
+
+void FXmlArchiveOutputFormatter::EnterAttributedValue()
+{
+	pugi::xml_node ValueElement = Impl->XmlElement.append_child(pugi::node_element);
+	ValueElement.set_name("AttributedValue");
+	Impl->XmlElement = ValueElement;
+}
+
+void FXmlArchiveOutputFormatter::EnterAttribute(FArchiveFieldName AttributeName)
+{
+	pugi::xml_node AttributeElement = Impl->XmlElement.append_child(pugi::node_element);
+	AttributeElement.set_name("Attribute");
+	AttributeElement.append_attribute("Name") = STRING_CAST_UE_TO_PUGI(AttributeName.Name);
+	Impl->XmlElement = AttributeElement;
+	
+	//checkf(FCString::Strcmp(AttributeName.Name, TEXT("Value")) != 0, TEXT("Attributes called 'Value' are reserved by the implementation"));
+	//Impl->XmlAttrib = Impl->XmlElement.append_attribute(AttributeName.Name);
+}
+
+void FXmlArchiveOutputFormatter::LeaveAttribute()
+{
+	using namespace UE::XmlSerialization::Private;
+	// It seems LeaveAttribute can be called after EnterAttributedValueValue.
+	// So the sequence is:
+	// EnterAttribute, EnterAttributedValueValue, LeaveAttribute.
+	// 
+	// Since AttributedValueValue and Attribute are at the same level under AttributedValue,
+	// We can pop them interchangeably if requested.
+	if (Utils::IsElementA(Impl->XmlElement, TEXT("AttributedValueValue")))
+	{
+		Impl->Pop(TEXT("LeaveAttribute"), TEXT("AttributedValueValue"));
+	}
+	else
+	{
+		Impl->Pop(TEXT("LeaveAttribute"), TEXT("Attribute"));
+	}
+	//Impl->XmlAttrib = pugi::xml_attribute();
+}
+
+void FXmlArchiveOutputFormatter::LeaveAttributedValue()
+{
+	if (FCStringAnsi::Strcmp(Impl->XmlElement.name(), "AttributedValueValue") == 0)
+	{
+		Impl->Pop(TEXT("LeaveAttributedValue"),TEXT("AttributedValueValue"));
+	}	
+	
+	Impl->Pop(TEXT("LeaveAttributedValue"), TEXT("AttributedValue"));
+}
+
+void FXmlArchiveOutputFormatter::EnterAttributedValueValue()
+{
+	pugi::xml_node AttributedValueValueElement = Impl->XmlElement.append_child(pugi::node_element);
+	AttributedValueValueElement.set_name("AttributedValueValue");
+	Impl->XmlElement = AttributedValueValueElement;
+}
+
+bool FXmlArchiveOutputFormatter::TryEnterAttributedValueValue()
+{
+	return false;	// That's how it is in json as well.
+}
+
+bool FXmlArchiveOutputFormatter::TryEnterAttribute(FArchiveFieldName AttributeName, bool bEnterWhenSaving)
+{
+	if (bEnterWhenSaving)
+	{
+		EnterAttribute(AttributeName);
+	}
+	return bEnterWhenSaving;
+}
+
+void FXmlArchiveOutputFormatter::Serialize(uint8& Value)
+{
+	Impl->SetValue(Value);
+}
+
+void FXmlArchiveOutputFormatter::Serialize(uint16& Value)
+{
+	Impl->SetValue(Value);
+}
+
+void FXmlArchiveOutputFormatter::Serialize(uint32& Value)
+{
+	Impl->SetValue(Value);
+}
+
+void FXmlArchiveOutputFormatter::Serialize(uint64& Value)
+{
+	Impl->SetValue(Value);
+}
+
+void FXmlArchiveOutputFormatter::Serialize(int8& Value)
+{
+	Impl->SetValue(Value);
+}
+
+void FXmlArchiveOutputFormatter::Serialize(int16& Value)
+{
+	Impl->SetValue(Value);
+}
+
+void FXmlArchiveOutputFormatter::Serialize(int32& Value)
+{
+	Impl->SetValue(Value);
+}
+
+void FXmlArchiveOutputFormatter::Serialize(int64& Value)
+{
+	Impl->SetValue(Value);
+}
+
+void FXmlArchiveOutputFormatter::Serialize(float& Value)
+{
+	if ((float)(int)Value == Value)
+	{
+		Impl->SetValue((int)Value);
+	}
+	else
+	{
+		FString String = FString::Printf(TEXT("%.17g"), Value);
+#if DO_GUARD_SLOW
+		float RoundTripped;
+		LexFromString(RoundTripped, *String);
+		check(RoundTripped == Value);
+#endif
+		Impl->SetValue(*String);
+	}
+}
+
+void FXmlArchiveOutputFormatter::Serialize(double& Value)
+{
+	if ((double)(int)Value == Value)
+	{
+		Impl->SetValue((int)Value);
+	}
+	else
+	{
+		FString String = FString::Printf(TEXT("%.17g"), Value);
+#if DO_GUARD_SLOW
+		double RoundTripped;
+		LexFromString(RoundTripped, *String);
+		check(RoundTripped == Value);
+#endif
+		Impl->SetValue(*String);
+	}
+}
+
+void FXmlArchiveOutputFormatter::Serialize(bool& Value)
+{
+	Impl->SetValue(Value);
+}
+
+void FXmlArchiveOutputFormatter::Serialize(FString& Value)
+{
+	Impl->SetValue(*Value);
+	
+	// Insert a "String:" prefix to prevent incorrect interpretation as another explicit type
+	// if (Value.StartsWith(TEXT("Object:")) || Value.StartsWith(TEXT("String:")) || Value.StartsWith(TEXT("Base64:")))
+	// {
+	// 	SerializeStringInternal(FString::Printf(TEXT("String:%s"), *Value));
+	// }
+	// else
+	// {
+	// 	SerializeStringInternal(Value);
+	// }
+}
+
+void FXmlArchiveOutputFormatter::Serialize(FName& Value)
+{
+	Impl->SetValue(*Value.ToString());
+}
+
+void FXmlArchiveOutputFormatter::Serialize(UObject*& Value)
+{
+	using namespace UE::XmlSerialization::Private;
+	if (bSerializeObjectsInPlace)
+	{
+		if (Value != nullptr)
+		{
+			{
+				Impl->XmlElement.append_child(pugi::node_comment).set_value("Object In Place");
+			}
+			
+			FStructuredArchive ChildArchive(*this);
+			FStructuredArchiveRecord ObjectRecord = ChildArchive.Open().EnterRecord();
+			// For current record, add attribute for object type
+			const TSoftObjectPtr<UClass> ObjectClass =  Value->GetClass();
+			
+			Impl->XmlElement.append_attribute("Class") = STRING_CAST_UE_TO_PUGI(*ObjectClass.ToSoftObjectPath().ToString());
+			Impl->XmlElement.append_attribute("Name") = STRING_CAST_UE_TO_PUGI(*Value->GetName());
+			Impl->XmlElement.append_attribute("Flags").set_value(static_cast<uint32>(Value->GetFlags()));			
+			FormatterHelper::SerializeObject(ObjectRecord, Value);
+			ChildArchive.Close();
+		}
+		else
+		{
+			Impl->SetValue(TEXT("null"));
+		}
+	}
+	else
+	{
+		if (Value != nullptr && IsObjectAllowed(Value))
+		{
+			FPackageIndex ObjectIndex = ObjectIndicesMap->FindChecked(Value);
+			Impl->SetValue(*LexToString(ObjectIndex));
+		}
+		else
+		{
+			Impl->SetValue(TEXT("null"));
+		}
+	}
+}
+
+void FXmlArchiveOutputFormatter::Serialize(FText& Value)
+{
+	// Unlike FJsonArchiveOutputFormatter, my class is not friends with FText. :(
+	
+	// FStructuredArchive ChildArchive(*this);
+	// FText::SerializeText(ChildArchive.Open(), Value);	---> is Private
+	// ChildArchive.Close();
+}
+
+void FXmlArchiveOutputFormatter::Serialize(FWeakObjectPtr& Value)
+{
+	UObject* Ptr = Value.IsValid() ? Value.Get() : nullptr;
+	Serialize(Ptr);
+}
+
+void FXmlArchiveOutputFormatter::Serialize(FSoftObjectPtr& Value)
+{
+	if (bSerializeSoftObjectsInPlace)
+	{
+		UObject* Ptr = Value.IsValid() ? Value.Get() : nullptr;
+		Serialize(Ptr);
+	}
+	else
+	{
+		FSoftObjectPath Path = Value.ToSoftObjectPath();
+		Serialize(Path);
+	}
+}
+
+void FXmlArchiveOutputFormatter::Serialize(FSoftObjectPath& Value)
+{
+	if (Value.IsValid())
+	{
+		Impl->SetValue(*FString::Printf(TEXT("Object:%s"), *Value.ToString()));
+	}
+	else
+	{
+		Impl->SetValue(TEXT("null"));
+	}
+}
+
+void FXmlArchiveOutputFormatter::Serialize(FLazyObjectPtr& Value)
+{
+	if (Value.IsValid() && IsObjectAllowed(Value.Get()))
+	{
+		Impl->SetValue(*FString::Printf(TEXT("Lazy:%s"), *Value.GetUniqueID().ToString()));
+	}
+	else
+	{
+		Impl->SetValue(TEXT("null"));
+	}
+}
+
+void FXmlArchiveOutputFormatter::Serialize(FObjectPtr& Value)
+{
+	UObject* Object = Value.Get();
+	Serialize(Object);
+}
+
+void FXmlArchiveOutputFormatter::Serialize(TArray<uint8>& Data)
+{
+	Serialize(Data.GetData(), Data.Num());
+}
+
+void FXmlArchiveOutputFormatter::Serialize(void* Data, uint64 DataSize)
+{
+	// Encode the data on a single line. No need for hashing; intra-line merge conflicts are rare.
+	Impl->SetValue(*FString::Printf(TEXT("Base64:%s"), *FBase64::Encode((const uint8*)Data, DataSize)));
+}
+
+bool FXmlArchiveOutputFormatter::IsObjectAllowed(UObject* InObject) const
+{
+	return ObjectIndicesMap && ObjectIndicesMap->Contains(InObject);
+}
+#endif
diff --git a/Engine/Plugins/Experimental/SerializationUtils/Source/XmlSerialization/Private/Formatters/XmlFormatterHelper.cpp b/Engine/Plugins/Experimental/SerializationUtils/Source/XmlSerialization/Private/Formatters/XmlFormatterHelper.cpp
new file mode 100644
--- /dev/null
+++ b/Engine/Plugins/Experimental/SerializationUtils/Source/XmlSerialization/Private/Formatters/XmlFormatterHelper.cpp
@@ -0,0 +1,48 @@
+// Copyright Epic Games, Inc. All Rights Reserved.
+
+#include "XmlFormatterHelper.h"
+#include "UObject/Object.h"
+#include "UObject/UnrealType.h"
+#include "Serialization/SerializedPropertyScope.h"
+
+namespace UE::XmlSerialization::Private
+{
+	void FormatterHelper::SerializeObject(FStructuredArchiveRecord& InRootRecord
+		, UObject* InObject
+		, const TFunction<bool(FProperty*)>& ShouldSkipProperty)
+	{
+		if (!InObject)
+		{
+			return;
+		}
+	
+		FStructuredArchiveRecord ObjectRecord = InRootRecord.EnterRecord(*InObject->GetName());
+		FArchive& UnderlyingArchive = InRootRecord.GetUnderlyingArchive();
+	
+		for (FProperty* Property = InObject->GetClass()->PropertyLink; Property; Property = Property->PropertyLinkNext)
+		{
+			if (!Property || Property->HasAnyPropertyFlags(CPF_Transient) || (ShouldSkipProperty && ShouldSkipProperty(Property)))
+			{
+				continue;
+			}
+
+			const FStructuredArchiveSlot PropertyField = ObjectRecord.EnterField(*Property->GetName());
+		
+			for (int32 Index = 0; Index < Property->ArrayDim; ++Index)
+			{
+				UObject* const Archetype = InObject->GetArchetype();
+			
+				void* const Data = Property->ContainerPtrToValuePtr<void>(InObject, Index);
+				const void* const Defaults = Property->ContainerPtrToValuePtrForDefaults<void>(Archetype->GetClass()
+					, Archetype
+					, Index);
+			
+				//if (!Property->Identical(Target, Default, UnderlyingArchive.GetPortFlags()))
+				{
+					FSerializedPropertyScope SerializedProperty(UnderlyingArchive, Property);
+					Property->SerializeItem(PropertyField, Data, Defaults);
+				}
+			}
+		}
+	}
+}
diff --git a/Engine/Plugins/Experimental/SerializationUtils/Source/XmlSerialization/Private/Formatters/XmlFormatterHelper.h b/Engine/Plugins/Experimental/SerializationUtils/Source/XmlSerialization/Private/Formatters/XmlFormatterHelper.h
new file mode 100644
--- /dev/null
+++ b/Engine/Plugins/Experimental/SerializationUtils/Source/XmlSerialization/Private/Formatters/XmlFormatterHelper.h
@@ -0,0 +1,15 @@
+// Copyright Epic Games, Inc. All Rights Reserved.
+
+#pragma once
+
+#include "Templates/Function.h"
+
+class FProperty;
+class FStructuredArchiveRecord;
+class UObject;
+
+namespace UE::XmlSerialization::Private::FormatterHelper
+{
+	void SerializeObject(FStructuredArchiveRecord& InRootRecord, UObject* InObject
+		, const TFunction<bool(FProperty*)>& ShouldSkipProperty = TFunction<bool(FProperty*)>());
+}
\ No newline at end of file
diff --git a/Engine/Plugins/Experimental/SerializationUtils/Source/XmlSerialization/Private/Utils/XmlUtils.h b/Engine/Plugins/Experimental/SerializationUtils/Source/XmlSerialization/Private/Utils/XmlUtils.h
new file mode 100644
--- /dev/null
+++ b/Engine/Plugins/Experimental/SerializationUtils/Source/XmlSerialization/Private/Utils/XmlUtils.h
@@ -0,0 +1,84 @@
+// Copyright Epic Games, Inc. All Rights Reserved.
+
+#pragma once
+
+#include "Containers/StringConv.h"
+#include "Containers/UnrealString.h"
+#include "CoreTypes.h"
+#include "Misc/Build.h"
+#include "Misc/CString.h"
+#include "XmlSerializationDefines.h"
+
+#include "MaterialXFormat/PugiXML/pugixml.hpp"
+
+/**
+ * Pugi can be compiled either in WChar mode (recommended) or Ansi mode.
+ * Since we use the version shipped in MaterialXFormat (for now), we don't have a choice
+ * to use that config, i.e. ainsi mode, but ideally we could have an independent
+ * pugi third party lib in wchar mode. In the mean time, these macro are used to
+ * make the code easy to change.
+ */
+#ifdef PUGIXML_WCHAR_MODE
+#	define	STRING_CAST_UE_TO_PUGI(InString) InString
+#	define STRING_CAST_PUGI_TO_UE(InString) InString
+#else
+#	define STRING_CAST_UE_TO_PUGI(InString) StringCast<ANSICHAR>(InString).Get() 
+#	define STRING_CAST_PUGI_TO_UE(InString) StringCast<TCHAR>(InString).Get()
+#endif
+
+namespace UE::XmlSerialization::Private::Utils
+{
+	// Helper for pretty log
+	inline FString GetInfo(const pugi::xml_node& XmlElement)
+	{
+		const pugi::xml_attribute NameAttrib = XmlElement.attribute("Name");
+		if (!NameAttrib.empty())
+		{
+			return FString::Printf(TEXT("%s Name=\"%s\""), STRING_CAST_PUGI_TO_UE(XmlElement.name()), STRING_CAST_PUGI_TO_UE(NameAttrib.as_string()));
+		}
+		else
+		{
+			return FString(XmlElement.name());
+		}
+	}
+
+	inline bool IsElementA(const pugi::xml_node& Element, const TCHAR* TagName)
+	{
+		return (FCStringAnsi::Strcmp(Element.name(), STRING_CAST_UE_TO_PUGI(TagName)) == 0) ? true : false;
+	}
+
+	inline bool IsElementShouldBe(const pugi::xml_node& Element, const TCHAR* TagName)
+	{
+		const pugi::xml_attribute HackAttrib = Element.attribute("__HACK_SHOULD_BE__");
+		if (!HackAttrib.empty())
+		{
+			return (FCStringAnsi::Strcmp(HackAttrib.as_string(), STRING_CAST_UE_TO_PUGI(TagName)) == 0) ? true : false;
+		}
+		return false;
+	}
+	
+	inline pugi::xml_encoding ToPugiEncoding(EXmlSerializationEncoding InEncoding)
+	{
+		switch (InEncoding)
+		{
+		case EXmlSerializationEncoding::Utf8:
+			return pugi::encoding_utf8;
+		case EXmlSerializationEncoding::Utf16_Le:
+			return pugi::encoding_utf16_le;
+		case EXmlSerializationEncoding::Utf16_Be:
+			return pugi::encoding_utf16_be;
+		case EXmlSerializationEncoding::Utf16:
+			return pugi::encoding_utf16;
+		case EXmlSerializationEncoding::Utf32_Le:
+			return pugi::encoding_utf32_le;
+		case EXmlSerializationEncoding::Utf32_Be:
+			return pugi::encoding_utf32_be;
+		case EXmlSerializationEncoding::Utf32:
+			return pugi::encoding_utf32;
+		case EXmlSerializationEncoding::WChar:
+			return pugi::encoding_wchar;
+		default:
+			return pugi::encoding_utf8;
+		}
+	}
+}
\ No newline at end of file
diff --git a/Engine/Plugins/Experimental/SerializationUtils/Source/XmlSerialization/Private/Utils/XmlWriter.h b/Engine/Plugins/Experimental/SerializationUtils/Source/XmlSerialization/Private/Utils/XmlWriter.h
new file mode 100644
--- /dev/null
+++ b/Engine/Plugins/Experimental/SerializationUtils/Source/XmlSerialization/Private/Utils/XmlWriter.h
@@ -0,0 +1,22 @@
+// Copyright Epic Games, Inc. All Rights Reserved.
+
+#pragma once
+
+#include "Serialization/Archive.h"
+#include "MaterialXFormat/PugiXML/pugixml.hpp"
+
+class FXmlWriter : public pugi::xml_writer
+{
+public:
+	FXmlWriter(	FArchive* const InStream) : Stream(InStream) {}
+
+	virtual void write(const void* data, size_t size) override
+	{
+		if (Stream != nullptr)
+		{
+			Stream->Serialize(const_cast<void*>(data), size);
+		}
+	}
+private:
+	FArchive* const Stream;
+};
diff --git a/Engine/Plugins/Experimental/SerializationUtils/Source/XmlSerialization/Private/XmlSerializationModule.cpp b/Engine/Plugins/Experimental/SerializationUtils/Source/XmlSerialization/Private/XmlSerializationModule.cpp
new file mode 100644
--- /dev/null
+++ b/Engine/Plugins/Experimental/SerializationUtils/Source/XmlSerialization/Private/XmlSerializationModule.cpp
@@ -0,0 +1,16 @@
+// Copyright Epic Games, Inc. All Rights Reserved.
+
+#include "XmlSerializationModule.h"
+#include "Modules/ModuleManager.h"
+
+DEFINE_LOG_CATEGORY(LogXmlSerialization);
+
+void FXmlSerializationModule::StartupModule()
+{
+}
+
+void FXmlSerializationModule::ShutdownModule()
+{
+}
+
+IMPLEMENT_MODULE(FXmlSerializationModule, XmlSerialization)
\ No newline at end of file
diff --git a/Engine/Plugins/Experimental/SerializationUtils/Source/XmlSerialization/Private/XmlSerializationModule.h b/Engine/Plugins/Experimental/SerializationUtils/Source/XmlSerialization/Private/XmlSerializationModule.h
new file mode 100644
--- /dev/null
+++ b/Engine/Plugins/Experimental/SerializationUtils/Source/XmlSerialization/Private/XmlSerializationModule.h
@@ -0,0 +1,17 @@
+// Copyright Epic Games, Inc. All Rights Reserved.
+
+#pragma once
+
+#include "CoreMinimal.h"
+#include "Modules/ModuleInterface.h"
+
+DECLARE_LOG_CATEGORY_EXTERN(LogXmlSerialization, Log, All);
+
+class FXmlSerializationModule : public IModuleInterface
+{
+public:
+	//~ Begin IModuleInterface
+	virtual void StartupModule() override;
+	virtual void ShutdownModule() override;
+	//~ End IModuleInterface
+};
diff --git a/Engine/Plugins/Experimental/SerializationUtils/Source/XmlSerialization/Public/Backends/XmlStructSerializerBackend.h b/Engine/Plugins/Experimental/SerializationUtils/Source/XmlSerialization/Public/Backends/XmlStructSerializerBackend.h
new file mode 100644
--- /dev/null
+++ b/Engine/Plugins/Experimental/SerializationUtils/Source/XmlSerialization/Public/Backends/XmlStructSerializerBackend.h
@@ -0,0 +1,47 @@
+// Copyright Epic Games, Inc. All Rights Reserved.
+
+#pragma once
+
+#include "IStructSerializerBackend.h"
+#include "XmlSerializationDefines.h"
+
+/**
+ * Implements a writer for UStruct serialization using Xml.
+ */
+class XMLSERIALIZATION_API FXmlStructSerializerBackend
+	: public IStructSerializerBackend
+{
+public:
+	
+	/**
+	 * Creates and initializes a new instance with the given flags.
+	 *
+	 * @param InArchive The archive to serialize into.
+	 * @param InFlags The flags that control the serialization behavior (typically EStructSerializerBackendFlags::Default).
+	 */
+	FXmlStructSerializerBackend( FArchive& InArchive, const EStructSerializerBackendFlags InFlags );
+
+	virtual ~FXmlStructSerializerBackend() override;
+
+public:
+	
+	/** Flush to archive. */
+	void SaveDocument(EXmlSerializationEncoding InEncoding = EXmlSerializationEncoding::Utf8);
+
+	//~ Begin IStructSerializerBackend
+	virtual void BeginArray(const FStructSerializerState& InState) override;
+	virtual void BeginStructure(const FStructSerializerState& InState) override;
+	virtual void EndArray(const FStructSerializerState& InState) override;
+	virtual void EndStructure(const FStructSerializerState& InState) override;
+	virtual void WriteComment(const FString& InComment) override;
+	virtual void WriteProperty(const FStructSerializerState& InState, int32 InArrayIndex) override;
+	//~ End IStructSerializerBackend
+	
+private:
+	// We keep implementation private so we don't expose the xml lib used.
+	class FImpl;
+	FImpl* Impl; 
+	
+	/** Flags controlling the serialization behavior. */
+	EStructSerializerBackendFlags Flags;
+};
diff --git a/Engine/Plugins/Experimental/SerializationUtils/Source/XmlSerialization/Public/Formatters/XmlArchiveInputFormatter.h b/Engine/Plugins/Experimental/SerializationUtils/Source/XmlSerialization/Public/Formatters/XmlArchiveInputFormatter.h
new file mode 100644
--- /dev/null
+++ b/Engine/Plugins/Experimental/SerializationUtils/Source/XmlSerialization/Public/Formatters/XmlArchiveInputFormatter.h
@@ -0,0 +1,98 @@
+// Copyright Epic Games, Inc. All Rights Reserved.
+
+#pragma once
+
+#include "CoreMinimal.h"
+#include "Serialization/StructuredArchiveFormatter.h"
+#include "UObject/ObjectResource.h"
+
+#if WITH_TEXT_ARCHIVE_SUPPORT
+
+class XMLSERIALIZATION_API FXmlArchiveInputFormatter final : public FStructuredArchiveFormatter
+{
+public:
+	/**
+	 * Parse the given archive into an xml DOM for deserializing objects.
+	 *
+	 * @param InInner Inner archive to read from. The whole archive is read in the constructor to parse into xml DOM.
+	 * @param InRootObject Specifies the root object to be the owner for nested objects. Can't be null.
+	 * @param InResolveObject Resolves reference for the UObject properties. This is not used if objects where serialized nested. 
+	 */
+	FXmlArchiveInputFormatter(FArchive& InInner, UObject* InRootObject, const TFunction<UObject* (const FPackageIndex)>& InResolveObject = nullptr);
+	FXmlArchiveInputFormatter(const FXmlArchiveInputFormatter& InOther);
+
+	virtual ~FXmlArchiveInputFormatter() override;
+
+	bool IsParseResultStatusOk() const;
+
+	virtual FArchive& GetUnderlyingArchive() override;
+	virtual FStructuredArchiveFormatter* CreateSubtreeReader() override;
+
+	virtual bool HasDocumentTree() const override;
+
+	virtual void EnterRecord() override;
+	virtual void LeaveRecord() override;
+	virtual void EnterField(FArchiveFieldName InName) override;
+	virtual void LeaveField() override;
+	virtual bool TryEnterField(FArchiveFieldName InName, bool bInEnterWhenSaving) override;
+
+	virtual void EnterArray(int32& OutNumElements) override;
+	virtual void LeaveArray() override;
+	virtual void EnterArrayElement() override;
+	virtual void LeaveArrayElement() override;
+
+	virtual void EnterStream() override;
+	virtual void LeaveStream() override;
+	virtual void EnterStreamElement() override;
+	virtual void LeaveStreamElement() override;
+
+	virtual void EnterMap(int32& OutNumElements) override;
+	virtual void LeaveMap() override;
+	virtual void EnterMapElement(FString& OutName) override;
+	virtual void LeaveMapElement() override;
+
+	virtual void EnterAttributedValue() override;
+	virtual void EnterAttribute(FArchiveFieldName InAttributeName) override;
+	virtual void EnterAttributedValueValue() override;
+	virtual void LeaveAttribute() override;
+	virtual void LeaveAttributedValue() override;
+	virtual bool TryEnterAttribute(FArchiveFieldName InAttributeName, bool bInEnterWhenSavin) override;
+	virtual bool TryEnterAttributedValueValue() override;
+
+	virtual void Serialize(uint8& Value) override;
+	virtual void Serialize(uint16& Value) override;
+	virtual void Serialize(uint32& Value) override;
+	virtual void Serialize(uint64& Value) override;
+	virtual void Serialize(int8& Value) override;
+	virtual void Serialize(int16& Value) override;
+	virtual void Serialize(int32& Value) override;
+	virtual void Serialize(int64& Value) override;
+	virtual void Serialize(float& Value) override;
+	virtual void Serialize(double& Value) override;
+	virtual void Serialize(bool& Value) override;
+	virtual void Serialize(FString& Value) override;
+	virtual void Serialize(FName& Value) override;
+	virtual void Serialize(UObject*& Value) override;
+	virtual void Serialize(FText& Value) override;
+	virtual void Serialize(FWeakObjectPtr& Value) override;
+	virtual void Serialize(FSoftObjectPtr& Value) override;
+	virtual void Serialize(FSoftObjectPath& Value) override;
+	virtual void Serialize(FLazyObjectPtr& Value) override;
+	virtual void Serialize(FObjectPtr& Value) override;
+	virtual void Serialize(TArray<uint8>& Value) override;
+	virtual void Serialize(void* Data, uint64 DataSize) override;
+
+private:
+	bool IsNestedObject() const;
+	UObject* LoadNestedObject();
+
+private:
+	class FImpl;
+	FImpl* Impl;
+	
+	FArchive& Inner;
+	TFunction<UObject* (const FPackageIndex)> ResolveObject;
+	UObject* RootObject;
+};
+
+#endif
diff --git a/Engine/Plugins/Experimental/SerializationUtils/Source/XmlSerialization/Public/Formatters/XmlArchiveOutputFormatter.h b/Engine/Plugins/Experimental/SerializationUtils/Source/XmlSerialization/Public/Formatters/XmlArchiveOutputFormatter.h
new file mode 100644
--- /dev/null
+++ b/Engine/Plugins/Experimental/SerializationUtils/Source/XmlSerialization/Public/Formatters/XmlArchiveOutputFormatter.h
@@ -0,0 +1,110 @@
+// Copyright Epic Games, Inc. All Rights Reserved.
+
+#pragma once
+
+#include "Serialization/StructuredArchiveFormatter.h"
+#include "UObject/ObjectResource.h"
+#include "XmlSerializationDefines.h"
+
+#if WITH_TEXT_ARCHIVE_SUPPORT
+
+class XMLSERIALIZATION_API FXmlArchiveOutputFormatter final : public FStructuredArchiveFormatter
+{
+private:
+	// Non-copyable - Need to implement if needed.
+	FXmlArchiveOutputFormatter(FXmlArchiveOutputFormatter&) = delete;
+	FXmlArchiveOutputFormatter& operator=(const FXmlArchiveOutputFormatter&) = delete;
+	// Non-movable - Need to implement if needed.
+	FXmlArchiveOutputFormatter(FXmlArchiveOutputFormatter&&) = delete;
+	FXmlArchiveOutputFormatter& operator=(FXmlArchiveOutputFormatter&&) = delete;
+
+public:
+	/**
+	 *	@remark Inner archive must remain valid throughout the lifetime of this object.
+	 */
+	FXmlArchiveOutputFormatter(FArchive& InInner);
+	
+	virtual ~FXmlArchiveOutputFormatter() override;
+
+	bool SaveDocumentToInnerArchive(EXmlSerializationEncoding InEncoding = EXmlSerializationEncoding::Utf8);
+	
+	void SerializeObjectsInPlace(bool bEnabled) { bSerializeObjectsInPlace = bEnabled; }
+	void SerializeSoftObjectsInPlace(bool bEnabled) { bSerializeSoftObjectsInPlace = bEnabled; }
+
+	virtual FArchive& GetUnderlyingArchive() override;
+
+	virtual bool HasDocumentTree() const override;
+
+	virtual void EnterRecord() override;
+	virtual void LeaveRecord() override;
+	virtual void EnterField(FArchiveFieldName Name) override;
+	virtual void LeaveField() override;
+	virtual bool TryEnterField(FArchiveFieldName Name, bool bEnterWhenSaving) override;
+
+	virtual void EnterArray(int32& NumElements) override;
+	virtual void LeaveArray() override;
+	virtual void EnterArrayElement() override;
+	virtual void LeaveArrayElement() override;
+
+	virtual void EnterStream() override;
+	virtual void LeaveStream() override;
+	virtual void EnterStreamElement() override;
+	virtual void LeaveStreamElement() override;
+
+	virtual void EnterMap(int32& NumElements) override;
+	virtual void LeaveMap() override;
+	virtual void EnterMapElement(FString& Name) override;
+	virtual void LeaveMapElement() override;
+
+	virtual void EnterAttributedValue() override;
+	virtual void EnterAttribute(FArchiveFieldName AttributeName) override;
+	virtual void EnterAttributedValueValue() override;
+	virtual void LeaveAttribute() override;
+	virtual void LeaveAttributedValue() override;
+	virtual bool TryEnterAttribute(FArchiveFieldName AttributeName, bool bEnterWhenSaving) override;
+
+	virtual bool TryEnterAttributedValueValue() override;
+
+	virtual void Serialize(uint8& Value) override;
+	virtual void Serialize(uint16& Value) override;
+	virtual void Serialize(uint32& Value) override;
+	virtual void Serialize(uint64& Value) override;
+	virtual void Serialize(int8& Value) override;
+	virtual void Serialize(int16& Value) override;
+	virtual void Serialize(int32& Value) override;
+	virtual void Serialize(int64& Value) override;
+	virtual void Serialize(float& Value) override;
+	virtual void Serialize(double& Value) override;
+	virtual void Serialize(bool& Value) override;
+	virtual void Serialize(FString& Value) override;
+	virtual void Serialize(FName& Value) override;
+	virtual void Serialize(UObject*& Value) override;
+	virtual void Serialize(FText& Value) override;
+	virtual void Serialize(FWeakObjectPtr& Value) override;
+	virtual void Serialize(FSoftObjectPtr& Value) override;
+	virtual void Serialize(FSoftObjectPath& Value) override;
+	virtual void Serialize(FLazyObjectPtr& Value) override;
+	virtual void Serialize(FObjectPtr& Value) override;
+	virtual void Serialize(TArray<uint8>& Value) override;
+	virtual void Serialize(void* Data, uint64 DataSize) override;
+
+	void SetObjectIndicesMap(const TMap<UObject*, FPackageIndex>* InObjectIndicesMap)
+	{
+		ObjectIndicesMap = InObjectIndicesMap;
+	}
+
+private:
+	bool IsObjectAllowed(UObject* InObject) const;
+
+private:
+	class FImpl;
+	FImpl* Impl;
+	
+	FArchive& Inner;
+	
+	const TMap<UObject*, FPackageIndex>* ObjectIndicesMap = nullptr;
+	bool bSerializeObjectsInPlace = false;
+	bool bSerializeSoftObjectsInPlace = false;
+};
+
+#endif
diff --git a/Engine/Plugins/Experimental/SerializationUtils/Source/XmlSerialization/Public/XmlSerializationDefines.h b/Engine/Plugins/Experimental/SerializationUtils/Source/XmlSerialization/Public/XmlSerializationDefines.h
new file mode 100644
--- /dev/null
+++ b/Engine/Plugins/Experimental/SerializationUtils/Source/XmlSerialization/Public/XmlSerializationDefines.h
@@ -0,0 +1,24 @@
+// Copyright Epic Games, Inc. All Rights Reserved.
+
+#pragma once
+
+/** These flags determine the encoding of input data for XML document. */
+enum class EXmlSerializationEncoding
+{
+	/** UTF8 encoding */
+	Utf8,
+	/** Little-endian UTF16 */
+	Utf16_Le,
+	/** Big-endian UTF16 */
+	Utf16_Be,
+	/** UTF16 with native endianness */
+	Utf16,
+	/** Little-endian UTF32 */
+	Utf32_Le,
+	/** Big-endian UTF32 */
+	Utf32_Be,
+	/** UTF32 with native endianness */
+	Utf32,
+	/** The same encoding wchar_t has (either UTF16 or UTF32) */
+	WChar,
+};
\ No newline at end of file
diff --git a/Engine/Plugins/Experimental/SerializationUtils/Source/XmlSerialization/XmlSerialization.Build.cs b/Engine/Plugins/Experimental/SerializationUtils/Source/XmlSerialization/XmlSerialization.Build.cs
new file mode 100644
--- /dev/null
+++ b/Engine/Plugins/Experimental/SerializationUtils/Source/XmlSerialization/XmlSerialization.Build.cs
@@ -0,0 +1,32 @@
+// Copyright Epic Games, Inc. All Rights Reserved.
+
+using System.IO;
+using UnrealBuildTool;
+
+public class XmlSerialization : ModuleRules
+{
+	public XmlSerialization(ReadOnlyTargetRules Target) : base(Target)
+	{
+		PCHUsage = ModuleRules.PCHUsageMode.UseExplicitOrSharedPCHs;
+
+		PublicDependencyModuleNames.AddRange(
+			new string[]
+			{
+				"Core",
+				"CoreUObject",
+			}
+		);
+
+		// PrivateDefinitions.Add("PUGIXML_WCHAR_MODE"); // for pugixml wchar support
+
+		PrivateDependencyModuleNames.AddRange(
+			new string[]
+			{
+				"MaterialX", // for pugixml
+				"Serialization",
+			}
+		);
+
+		AddEngineThirdPartyPrivateStaticDependencies(Target, "MaterialX");
+	}
+}
