Motion Design:
- Updated plugins descriptions and set set to experimental

#changelist validated
#virtualized

diff --git a/Engine/Plugins/Experimental/ActorModifierCore/ActorModifierCore.uplugin b/Engine/Plugins/Experimental/ActorModifierCore/ActorModifierCore.uplugin
--- a/Engine/Plugins/Experimental/ActorModifierCore/ActorModifierCore.uplugin
+++ b/Engine/Plugins/Experimental/ActorModifierCore/ActorModifierCore.uplugin
@@ -2,11 +2,11 @@
 	"FileVersion": 3,
 	"Version": 1,
 	"VersionName": "1.0",
-	"FriendlyName": "ActorModifierCore",
+	"FriendlyName": "Actor Modifier Core",
 	"Description": "Use modifier objects on actors to apply a custom behavior",
-	"Category": "Modifier",
+	"Category": "Experimental",
 	"CreatedBy": "Epic Games, Inc.",
-	"CreatedByURL": "",
+	"CreatedByURL": "https://epicgames.com",
 	"DocsURL": "",
 	"MarketplaceURL": "",
 	"SupportURL": "",
diff --git a/Engine/Plugins/Experimental/AdvancedRenamer/AdvancedRenamer.uplugin b/Engine/Plugins/Experimental/AdvancedRenamer/AdvancedRenamer.uplugin
--- a/Engine/Plugins/Experimental/AdvancedRenamer/AdvancedRenamer.uplugin
+++ b/Engine/Plugins/Experimental/AdvancedRenamer/AdvancedRenamer.uplugin
@@ -1,17 +1,18 @@
 {
 	"FileVersion": 3,
-	"Version": 0.1,
-	"VersionName": "0.1.0-alpha",
+	"Version": 1,
+	"VersionName": "1.0",
 	"FriendlyName": "Advanced Renamer",
-	"Description": "Renaming actors and assets",
+	"Description": "Advanced Renaming tool for actors and assets",
 	"Category": "Experimental",
 	"CreatedBy": "Epic Games, Inc.",
-	"CreatedByURL": "http://epicgames.com",
+	"CreatedByURL": "https://epicgames.com",
 	"DocsURL": "",
 	"MarketplaceURL": "",
 	"SupportURL": "",
-	"CanContainContent": true,
-	"IsBetaVersion": true,
+	"CanContainContent": false,
+	"IsBetaVersion": false,
+	"IsExperimentalVersion": true,
 	"Installed": false,
 	"Modules": [
 		{
diff --git a/Engine/Plugins/Experimental/Avalanche/Avalanche.uplugin b/Engine/Plugins/Experimental/Avalanche/Avalanche.uplugin
--- a/Engine/Plugins/Experimental/Avalanche/Avalanche.uplugin
+++ b/Engine/Plugins/Experimental/Avalanche/Avalanche.uplugin
@@ -1,12 +1,12 @@
 {
 	"FileVersion": 3,
-	"Version": 0.1,
-	"VersionName": "0.1",
+	"Version": 1,
+	"VersionName": "1.0",
 	"FriendlyName": "Motion Design",
 	"Description": "Compositing, designer and broadcasting tool.\n\nPlugin Dependencies: Advanced Renamer, Custom Details View, Dynamic Material, Geometry Cache, Geometry Scripting, Media Compositing, Media IO Framework, Mesh Modeling Toolset Exp, Remote Control, SVG Importer, Text3D and ActorModifierCore.",
 	"Category": "Experimental",
 	"CreatedBy": "Epic Games, Inc.",
-	"CreatedByURL": "http://epicgames.com",
+	"CreatedByURL": "https://epicgames.com",
 	"DocsURL": "",
 	"MarketplaceURL": "",
 	"SupportURL": "",
diff --git a/Engine/Plugins/Experimental/ClonerEffector/ClonerEffector.uplugin b/Engine/Plugins/Experimental/ClonerEffector/ClonerEffector.uplugin
--- a/Engine/Plugins/Experimental/ClonerEffector/ClonerEffector.uplugin
+++ b/Engine/Plugins/Experimental/ClonerEffector/ClonerEffector.uplugin
@@ -11,7 +11,7 @@
 	"MarketplaceURL": "",
 	"CanContainContent": true,
 	"IsBetaVersion": false,
-	"IsExperimentalVersion": false,
+	"IsExperimentalVersion": true,
 	"Installed": false,
 	"Modules": [
 		{
diff --git a/Engine/Plugins/Experimental/CustomDetailsView/CustomDetailsView.uplugin b/Engine/Plugins/Experimental/CustomDetailsView/CustomDetailsView.uplugin
--- a/Engine/Plugins/Experimental/CustomDetailsView/CustomDetailsView.uplugin
+++ b/Engine/Plugins/Experimental/CustomDetailsView/CustomDetailsView.uplugin
@@ -2,17 +2,17 @@
 	"FileVersion": 1,
 	"Version": 1,
 	"VersionName": "1.0",
-	"FriendlyName": "CustomDetailsView",
+	"FriendlyName": "Custom Details View",
 	"Description": "",
-	"Category": "Other",
-	"CreatedBy": "",
-	"CreatedByURL": "",
+	"Category": "Experimental",
+	"CreatedBy": "Epic Games, Inc.",
+	"CreatedByURL": "https://epicgames.com",
 	"DocsURL": "",
 	"MarketplaceURL": "",
 	"SupportURL": "",
-	"CanContainContent": true,
+	"CanContainContent": false,
 	"IsBetaVersion": false,
-	"IsExperimentalVersion": false,
+	"IsExperimentalVersion": true,
 	"Installed": false,
 	"Modules": [
 		{
diff --git a/Engine/Plugins/Experimental/DynamicMaterial/DynamicMaterial.uplugin b/Engine/Plugins/Experimental/DynamicMaterial/DynamicMaterial.uplugin
--- a/Engine/Plugins/Experimental/DynamicMaterial/DynamicMaterial.uplugin
+++ b/Engine/Plugins/Experimental/DynamicMaterial/DynamicMaterial.uplugin
@@ -1,17 +1,18 @@
 {
 	"FileVersion": 3,
-	"Version": 0.1,
-	"VersionName": "0.1.0-alpha",
+	"Version": 1,
+	"VersionName": "1.0",
 	"FriendlyName": "Material Designer",
 	"Description": "Compact dynamic material creator and editor, similar in style to other DDCs.",
 	"Category": "Experimental",
 	"CreatedBy": "Epic Games, Inc.",
-	"CreatedByURL": "http://epicgames.com",
+	"CreatedByURL": "https://epicgames.com",
 	"DocsURL": "",
 	"MarketplaceURL": "",
 	"SupportURL": "",
 	"CanContainContent": true,
-	"IsBetaVersion": true,
+	"IsBetaVersion": false,
+	"IsExperimentalVersion": true,
 	"Installed": false,
 	"Modules": [
 		{
diff --git a/Engine/Plugins/Experimental/FloatingProperties/FloatingProperties.uplugin b/Engine/Plugins/Experimental/FloatingProperties/FloatingProperties.uplugin
--- a/Engine/Plugins/Experimental/FloatingProperties/FloatingProperties.uplugin
+++ b/Engine/Plugins/Experimental/FloatingProperties/FloatingProperties.uplugin
@@ -1,17 +1,18 @@
 {
 	"FileVersion": 3,
-	"Version": 0.1,
-	"VersionName": "0.1.0-alpha",
+	"Version": 1,
+	"VersionName": "1.0",
 	"FriendlyName": "Floating Properties",
 	"Description": "Show floating properties from the details panel directly on the active viewport.",
 	"Category": "Experimental",
 	"CreatedBy": "Epic Games, Inc.",
-	"CreatedByURL": "http://epicgames.com",
+	"CreatedByURL": "https://epicgames.com",
 	"DocsURL": "",
 	"MarketplaceURL": "",
 	"SupportURL": "",
 	"CanContainContent": true,
-	"IsBetaVersion": true,
+	"IsBetaVersion": false,
+	"IsExperimentalVersion": true,
 	"Installed": false,
 	"Modules": [
 		{
diff --git a/Engine/Plugins/Experimental/GeometryMask/GeometryMask.uplugin b/Engine/Plugins/Experimental/GeometryMask/GeometryMask.uplugin
--- a/Engine/Plugins/Experimental/GeometryMask/GeometryMask.uplugin
+++ b/Engine/Plugins/Experimental/GeometryMask/GeometryMask.uplugin
@@ -2,16 +2,16 @@
 	"FileVersion": 3,
 	"Version": 1,
 	"VersionName": "1.0",
-	"FriendlyName": "GeometryMask",
+	"FriendlyName": "Geometry Mask",
 	"Description": "",
-	"Category": "Other",
-	"CreatedBy": "",
-	"CreatedByURL": "",
+	"Category": "Experimental",
+	"CreatedBy": "Epic Games, Inc.",
+	"CreatedByURL": "https://epicgames.com",
 	"DocsURL": "",
 	"MarketplaceURL": "",
 	"CanContainContent": true,
 	"IsBetaVersion": false,
-	"IsExperimentalVersion": false,
+	"IsExperimentalVersion": true,
 	"Installed": false,
 	"Modules": [
 		{
diff --git a/Engine/Plugins/Experimental/OperatorStack/OperatorStack.uplugin b/Engine/Plugins/Experimental/OperatorStack/OperatorStack.uplugin
--- a/Engine/Plugins/Experimental/OperatorStack/OperatorStack.uplugin
+++ b/Engine/Plugins/Experimental/OperatorStack/OperatorStack.uplugin
@@ -4,14 +4,14 @@
 	"VersionName": "1.0",
 	"FriendlyName": "Operator Stack",
 	"Description": "Allows you to display linked objects in a container with an intuitive UI",
-	"Category": "UI",
+	"Category": "Experimental",
 	"CreatedBy": "Epic Games, Inc.",
-	"CreatedByURL": "",
+	"CreatedByURL": "https://epicgames.com",
 	"DocsURL": "",
 	"MarketplaceURL": "",
-	"CanContainContent": true,
+	"CanContainContent": false,
 	"IsBetaVersion": false,
-	"IsExperimentalVersion": false,
+	"IsExperimentalVersion": true,
 	"Installed": false,
 	"Modules": [
 		{
diff --git a/Engine/Plugins/Experimental/PropertyAnimator/PropertyAnimator.uplugin b/Engine/Plugins/Experimental/PropertyAnimator/PropertyAnimator.uplugin
--- a/Engine/Plugins/Experimental/PropertyAnimator/PropertyAnimator.uplugin
+++ b/Engine/Plugins/Experimental/PropertyAnimator/PropertyAnimator.uplugin
@@ -1,17 +1,17 @@
 ﻿{
 	"FileVersion": 3,
 	"Version": 1,
-	"VersionName": "0.1",
-	"FriendlyName": "PropertyAnimator",
+	"VersionName": "1.0",
+	"FriendlyName": "Property Animator",
 	"Description": "Re-usable behaviors to animate the value of one or more properties",
-	"Category": "Animation",
+	"Category": "Experimental",
 	"CreatedBy": "Epic Games, Inc.",
 	"CreatedByURL": "https://epicgames.com",
 	"DocsURL": "",
 	"MarketplaceURL": "",
 	"CanContainContent": true,
 	"IsBetaVersion": false,
-	"IsExperimentalVersion": false,
+	"IsExperimentalVersion": true,
 	"Installed": false,
 	"Modules": [
 		{
diff --git a/Engine/Plugins/Experimental/PropertyAnimatorCore/PropertyAnimatorCore.uplugin b/Engine/Plugins/Experimental/PropertyAnimatorCore/PropertyAnimatorCore.uplugin
--- a/Engine/Plugins/Experimental/PropertyAnimatorCore/PropertyAnimatorCore.uplugin
+++ b/Engine/Plugins/Experimental/PropertyAnimatorCore/PropertyAnimatorCore.uplugin
@@ -2,16 +2,16 @@
 	"FileVersion": 3,
 	"Version": 1,
 	"VersionName": "1.0",
-	"FriendlyName": "PropertyAnimatorCore",
+	"FriendlyName": "Property Animator Core",
 	"Description": "Re-usable behaviors to control properties at runtime and in editor",
-	"Category": "Animation",
+	"Category": "Experimental",
 	"CreatedBy": "Epic Games, Inc.",
 	"CreatedByURL": "https://epicgames.com",
 	"DocsURL": "",
 	"MarketplaceURL": "",
 	"CanContainContent": true,
 	"IsBetaVersion": false,
-	"IsExperimentalVersion": false,
+	"IsExperimentalVersion": true,
 	"Installed": false,
 	"Modules": [
 		{
diff --git a/Engine/Plugins/Experimental/RemoteControlComponents/RemoteControlComponents.uplugin b/Engine/Plugins/Experimental/RemoteControlComponents/RemoteControlComponents.uplugin
--- a/Engine/Plugins/Experimental/RemoteControlComponents/RemoteControlComponents.uplugin
+++ b/Engine/Plugins/Experimental/RemoteControlComponents/RemoteControlComponents.uplugin
@@ -1,16 +1,17 @@
 ﻿{
 	"FileVersion": 3,
-	"Version": 0.1,
-	"VersionName": "0.1.0-alpha",
-	"FriendlyName": "RemoteControlComponents",
+	"Version": 1,
+	"VersionName": "1.0",
+	"FriendlyName": "Remote Control Components",
 	"Description": "",
 	"Category": "Experimental",
 	"CreatedBy": "Epic Games, Inc.",
-	"CreatedByURL": "http://epicgames.com",
+	"CreatedByURL": "https://epicgames.com",
 	"DocsURL": "",
 	"MarketplaceURL": "",
-	"CanContainContent": true,
-	"IsBetaVersion": true,
+	"CanContainContent": false,
+	"IsBetaVersion": false,
+	"IsExperimentalVersion": true,
 	"Installed": false,
 	"Modules": [
 		{
diff --git a/Engine/Plugins/Experimental/SVGImporter/SVGImporter.uplugin b/Engine/Plugins/Experimental/SVGImporter/SVGImporter.uplugin
--- a/Engine/Plugins/Experimental/SVGImporter/SVGImporter.uplugin
+++ b/Engine/Plugins/Experimental/SVGImporter/SVGImporter.uplugin
@@ -1,17 +1,17 @@
 {
 	"FileVersion": 3,
-	"Version": 0.1,
-	"VersionName": "0.1.0-alpha",
-	"FriendlyName": "SVGImporter",
+	"Version": 1,
+	"VersionName": "1.0",
+	"FriendlyName": "SVG Importer",
 	"Description": "Importing and handling SVG files",
 	"Category": "Experimental",
 	"CreatedBy": "Epic Games, Inc.",
-	"CreatedByURL": "http://epicgames.com",
+	"CreatedByURL": "https://epicgames.com",
 	"DocsURL": "",
 	"MarketplaceURL": "",
 	"SupportURL": "",
 	"CanContainContent": true,
-	"IsBetaVersion": true,
+	"IsBetaVersion": false,
 	"IsExperimentalVersion": true,
 	"Installed": false,
 	"Modules": [
diff --git a/Engine/Plugins/Experimental/SerializationUtils/SerializationUtils.uplugin b/Engine/Plugins/Experimental/SerializationUtils/SerializationUtils.uplugin
--- a/Engine/Plugins/Experimental/SerializationUtils/SerializationUtils.uplugin
+++ b/Engine/Plugins/Experimental/SerializationUtils/SerializationUtils.uplugin
@@ -1,18 +1,18 @@
 {
 	"FileVersion": 3,
-	"Version": 0.1,
-	"VersionName": "0.1.0-alpha",
-	"FriendlyName": "SerializationUtils",
+	"Version": 1,
+	"VersionName": "1.0",
+	"FriendlyName": "Serialization Utils",
 	"Description": "Utilities for serialization (xml, json, etc) with extended functionality.",
 	"Category": "Experimental",
 	"CreatedBy": "Epic Games, Inc.",
-	"CreatedByURL": "http://epicgames.com",
+	"CreatedByURL": "https://epicgames.com",
 	"DocsURL": "",
 	"MarketplaceURL": "",
 	"SupportURL": "",
-	"CanContainContent": true,
+	"CanContainContent": false,
 	"IsBetaVersion": false,
-	"IsExperimentalVersion": false,
+	"IsExperimentalVersion": true,
 	"Installed": false,
 	"Modules": [
 		{
diff --git a/Engine/Plugins/Experimental/StormSync/StormSync.uplugin b/Engine/Plugins/Experimental/StormSync/StormSync.uplugin
--- a/Engine/Plugins/Experimental/StormSync/StormSync.uplugin
+++ b/Engine/Plugins/Experimental/StormSync/StormSync.uplugin
@@ -1,18 +1,18 @@
 {
 	"FileVersion": 3,
 	"Version": 1,
-	"VersionName": "0.16.0",
-	"FriendlyName": "StormSync",
+	"VersionName": "1.0",
+	"FriendlyName": "Storm Sync",
 	"Description": "Sync, Pull, Push, asset dependencies.",
-	"Category": "Other",
+	"Category": "Experimental",
 	"CreatedBy": "Epic Games, Inc.",
-	"CreatedByURL": "",
+	"CreatedByURL": "https://epicgames.com",
 	"DocsURL": "",
 	"MarketplaceURL": "",
 	"SupportURL": "",
 	"CanContainContent": true,
 	"IsBetaVersion": false,
-	"IsExperimentalVersion": false,
+	"IsExperimentalVersion": true,
 	"Installed": false,
 	"Modules": [
 		{
diff --git a/Engine/Plugins/Experimental/StormSyncAvalancheBridge/StormSyncAvalancheBridge.uplugin b/Engine/Plugins/Experimental/StormSyncAvalancheBridge/StormSyncAvalancheBridge.uplugin
--- a/Engine/Plugins/Experimental/StormSyncAvalancheBridge/StormSyncAvalancheBridge.uplugin
+++ b/Engine/Plugins/Experimental/StormSyncAvalancheBridge/StormSyncAvalancheBridge.uplugin
@@ -1,17 +1,18 @@
 {
 	"FileVersion": 3,
 	"Version": 1,
-	"VersionName": "0.16.0",
+	"VersionName": "1.0",
 	"FriendlyName": "Storm Sync Motion Design Bridge",
 	"Description": "Plugin bridge between Motion Design Plugin and Storm Sync to provide in-editor integration to synchronize assets",
 	"Category": "Experimental",
 	"CreatedBy": "Epic Games, Inc.",
-	"CreatedByURL": "",
+	"CreatedByURL": "https://epicgames.com",
 	"DocsURL": "",
 	"MarketplaceURL": "",
 	"SupportURL": "",
 	"CanContainContent": true,
-	"IsBetaVersion": true,
+	"IsBetaVersion": false,
+	"IsExperimentalVersion": true,
 	"Installed": false,
 	"Modules": [
 		{
diff --git a/Engine/Plugins/Experimental/WebSocketMessaging/WebSocketMessaging.uplugin b/Engine/Plugins/Experimental/WebSocketMessaging/WebSocketMessaging.uplugin
--- a/Engine/Plugins/Experimental/WebSocketMessaging/WebSocketMessaging.uplugin
+++ b/Engine/Plugins/Experimental/WebSocketMessaging/WebSocketMessaging.uplugin
@@ -2,17 +2,17 @@
 	"FileVersion": 3,
 	"Version": 1,
 	"VersionName": "1.0",
-	"FriendlyName": "WebSocketMessaging",
+	"FriendlyName": "Web Socket Messaging",
 	"Description": "",
-	"Category": "Other",
-	"CreatedBy": "",
-	"CreatedByURL": "",
+	"Category": "Experimental",
+	"CreatedBy": "Epic Games, Inc.",
+	"CreatedByURL": "https://epicgames.com",
 	"DocsURL": "",
 	"MarketplaceURL": "",
 	"SupportURL": "",
 	"CanContainContent": false,
 	"IsBetaVersion": false,
-	"IsExperimentalVersion": false,
+	"IsExperimentalVersion": true,
 	"Installed": false,
 	"Modules": [
 		{
