Added missing includes.

diff --git a/Engine/Plugins/Experimental/SerializationUtils/Source/XmlSerialization/Private/Backends/XmlStructSerializerBackend.cpp b/Engine/Plugins/Experimental/SerializationUtils/Source/XmlSerialization/Private/Backends/XmlStructSerializerBackend.cpp
--- a/Engine/Plugins/Experimental/SerializationUtils/Source/XmlSerialization/Private/Backends/XmlStructSerializerBackend.cpp
+++ b/Engine/Plugins/Experimental/SerializationUtils/Source/XmlSerialization/Private/Backends/XmlStructSerializerBackend.cpp
@@ -6,6 +6,7 @@
 #include "MaterialXFormat/PugiXML/pugixml.hpp"
 #include "UObject/EnumProperty.h"
 #include "UObject/PropertyPortFlags.h"
+#include "UObject/StrProperty.h"
 #include "UObject/TextProperty.h"
 #include "UObject/UnrealType.h"
 #include "Utils/XmlUtils.h"
diff --git a/Engine/Plugins/VirtualProduction/RemoteControl/Source/RemoteControl/Private/RemoteControlPropertyHandle.cpp b/Engine/Plugins/VirtualProduction/RemoteControl/Source/RemoteControl/Private/RemoteControlPropertyHandle.cpp
--- a/Engine/Plugins/VirtualProduction/RemoteControl/Source/RemoteControl/Private/RemoteControlPropertyHandle.cpp
+++ b/Engine/Plugins/VirtualProduction/RemoteControl/Source/RemoteControl/Private/RemoteControlPropertyHandle.cpp
@@ -15,6 +15,7 @@
 #include "Materials/MaterialInterface.h"
 #include "Serialization/MemoryReader.h"
 #include "Serialization/MemoryWriter.h"
+#include "UObject/StrProperty.h"
 #include "UObject/TextProperty.h"
 
 #if WITH_EDITOR
diff --git a/Engine/Plugins/VirtualProduction/RemoteControl/Source/RemoteControlCommon/Public/RCPropertyUtilities.inl b/Engine/Plugins/VirtualProduction/RemoteControl/Source/RemoteControlCommon/Public/RCPropertyUtilities.inl
--- a/Engine/Plugins/VirtualProduction/RemoteControl/Source/RemoteControlCommon/Public/RCPropertyUtilities.inl
+++ b/Engine/Plugins/VirtualProduction/RemoteControl/Source/RemoteControlCommon/Public/RCPropertyUtilities.inl
@@ -8,6 +8,7 @@
 #include "Components/SceneComponent.h"
 #include "Serialization/BufferArchive.h"
 #include "Serialization/MemoryReader.h"
+#include "UObject/StrProperty.h"
 #include "UObject/WeakFieldPtr.h"
 
 #if WITH_EDITOR
