[Xml Serialization] Update build.cs to only include the dynamic linkage of MaterialX module

diff --git a/Engine/Plugins/Experimental/SerializationUtils/Source/XmlSerialization/XmlSerialization.Build.cs b/Engine/Plugins/Experimental/SerializationUtils/Source/XmlSerialization/XmlSerialization.Build.cs
--- a/Engine/Plugins/Experimental/SerializationUtils/Source/XmlSerialization/XmlSerialization.Build.cs
+++ b/Engine/Plugins/Experimental/SerializationUtils/Source/XmlSerialization/XmlSerialization.Build.cs
@@ -26,7 +26,5 @@ public class XmlSerialization : ModuleRules
 				"Serialization",
 			}
 		);
-
-		AddEngineThirdPartyPrivateStaticDependencies(Target, "MaterialX");
 	}
 }
