The pugi namespace was renamed mxpugi in the MaterialX lib to fix symbol collisions with a console SDK library, this change repalces MaterialX pugi usage with mxpugi.

diff --git a/Engine/Plugins/Experimental/SVGImporter/Source/SVGImporterEditor/Private/SVGParser_PugiXml.cpp b/Engine/Plugins/Experimental/SVGImporter/Source/SVGImporterEditor/Private/SVGParser_PugiXml.cpp
--- a/Engine/Plugins/Experimental/SVGImporter/Source/SVGImporterEditor/Private/SVGParser_PugiXml.cpp
+++ b/Engine/Plugins/Experimental/SVGImporter/Source/SVGImporterEditor/Private/SVGParser_PugiXml.cpp
@@ -11,13 +11,13 @@ bool FSVGParser_PugiXml::Parse(bool bInTryToImportWithError /* = true */)
 		ParseXML();
 	}
 
-	if (XmlParseResult.status == pugi::xml_parse_status::status_ok)
+	if (XmlParseResult.status == mxpugi::xml_parse_status::status_ok)
 	{
 		// Create root of SVG hierarchy
 		RootSVGElement = MakeShared<FSVGRawElement>();
 		SVGElements.Emplace(RootSVGElement);
 
-		if (const pugi::xml_node* DocPtr = XmlDoc.Get())
+		if (const mxpugi::xml_node* DocPtr = XmlDoc.Get())
 		{
 			ParseNode(*DocPtr, RootSVGElement);
 			return true;
@@ -45,29 +45,29 @@ bool FSVGParser_PugiXml::IsValidSVG()
 		ParseXML();
 	}
 
-	return XmlParseResult.status == pugi::xml_parse_status::status_ok;
+	return XmlParseResult.status == mxpugi::xml_parse_status::status_ok;
 }
 
 void FSVGParser_PugiXml::ParseXML()
 {
 	if (!XmlDoc)
 	{
-		XmlDoc = MakeShared<pugi::xml_document>();
+		XmlDoc = MakeShared<mxpugi::xml_document>();
 	}
 
 	// Parse xml from buffer
 	XmlParseResult = XmlDoc->load_buffer(GetData(SVGString), SVGString.GetAllocatedSize());
 }
 
-void FSVGParser_PugiXml::ParseNode(const pugi::xml_node& InNode, const TSharedPtr<FSVGRawElement>& InParent)
+void FSVGParser_PugiXml::ParseNode(const mxpugi::xml_node& InNode, const TSharedPtr<FSVGRawElement>& InParent)
 {
 	const TSharedPtr<FSVGRawElement> ParentElement = InParent.IsValid() ? InParent : RootSVGElement;
 
-	const pugi::char_t* NodeName = InNode.name();
+	const mxpugi::char_t* NodeName = InNode.name();
 
 	// node.child_value() is the way for elements to get non attributes values.
 	// For SVGs this can be found e.g in style tags: <style> VALUE </style>
-	const pugi::char_t* NodeValue = InNode.child_value();
+	const mxpugi::char_t* NodeValue = InNode.child_value();
 
 	// An element has been found, let's create a new USVGParsedElement, specifying its parent, name and data
 	const TSharedRef<FSVGRawElement> NewElement = FSVGRawElement::NewSVGElement(ParentElement, NodeName, NodeValue);
@@ -77,7 +77,7 @@ void FSVGParser_PugiXml::ParseNode(const pugi::xml_node& InNode, const TSharedPt
 		ParentElement->AddChild(NewElement);
 	}
 
-	for (pugi::xml_attribute_iterator AttrItr = InNode.attributes_begin(); AttrItr != InNode.attributes_end(); ++AttrItr)
+	for (mxpugi::xml_attribute_iterator AttrItr = InNode.attributes_begin(); AttrItr != InNode.attributes_end(); ++AttrItr)
 	{
 		// Attribute already exists, continue
 		if (NewElement->GetAttributes().Contains(AttrItr->name()))
@@ -90,7 +90,7 @@ void FSVGParser_PugiXml::ParseNode(const pugi::xml_node& InNode, const TSharedPt
 		NewElement->AddAttribute(NewAttribute);
 	}
 
-	for (const pugi::xml_node& ChildNode : InNode.children())
+	for (const mxpugi::xml_node& ChildNode : InNode.children())
 	{
 		ParseNode(ChildNode, NewElement);
 	}
diff --git a/Engine/Plugins/Experimental/SVGImporter/Source/SVGImporterEditor/Private/SVGParser_PugiXml.h b/Engine/Plugins/Experimental/SVGImporter/Source/SVGImporterEditor/Private/SVGParser_PugiXml.h
--- a/Engine/Plugins/Experimental/SVGImporter/Source/SVGImporterEditor/Private/SVGParser_PugiXml.h
+++ b/Engine/Plugins/Experimental/SVGImporter/Source/SVGImporterEditor/Private/SVGParser_PugiXml.h
@@ -25,11 +25,11 @@ protected:
 	void ParseXML();
 
 	/** Recursively parse a node and its children nodes */
-	void ParseNode(const pugi::xml_node& InNode, const TSharedPtr<FSVGRawElement>& InParent);
+	void ParseNode(const mxpugi::xml_node& InNode, const TSharedPtr<FSVGRawElement>& InParent);
 
 	/** The pugi xml document which will be used for parsing */
-	TSharedPtr<pugi::xml_document> XmlDoc;
+	TSharedPtr<mxpugi::xml_document> XmlDoc;
 
 	/** Result of xml parsing */
-	pugi::xml_parse_result XmlParseResult;
+	mxpugi::xml_parse_result XmlParseResult;
 };
diff --git a/Engine/Plugins/Experimental/SerializationUtils/Source/XmlSerialization/Private/Backends/XmlStructSerializerBackend.cpp b/Engine/Plugins/Experimental/SerializationUtils/Source/XmlSerialization/Private/Backends/XmlStructSerializerBackend.cpp
--- a/Engine/Plugins/Experimental/SerializationUtils/Source/XmlSerialization/Private/Backends/XmlStructSerializerBackend.cpp
+++ b/Engine/Plugins/Experimental/SerializationUtils/Source/XmlSerialization/Private/Backends/XmlStructSerializerBackend.cpp
@@ -16,7 +16,7 @@ class FXmlStructSerializerBackend::FImpl
 public:
 	FImpl()
 	{
-		pugi::xml_node decl = XmlDoc.prepend_child(pugi::node_declaration);
+		mxpugi::xml_node decl = XmlDoc.prepend_child(mxpugi::node_declaration);
 		decl.append_attribute("version") = "1.0";
 		decl.append_attribute("encoding") = "utf-8";
 
@@ -25,8 +25,8 @@ public:
 		XmlElement.append_attribute(("Version")) = "1.0";
 	}
 	
-	pugi::xml_document XmlDoc;
-	pugi::xml_node XmlElement;
+	mxpugi::xml_document XmlDoc;
+	mxpugi::xml_node XmlElement;
 	TUniquePtr<FXmlWriter> XmlWriter;
 };
 
@@ -46,13 +46,13 @@ FXmlStructSerializerBackend::~FXmlStructSerializerBackend()
 void FXmlStructSerializerBackend::SaveDocument(EXmlSerializationEncoding InEncoding)
 {
 	using namespace UE::XmlSerialization::Private;
-	const pugi::xml_encoding encoding = Utils::ToPugiEncoding(InEncoding);
-	Impl->XmlDoc.save(*Impl->XmlWriter, "  ", pugi::format_default, encoding);
+	const mxpugi::xml_encoding encoding = Utils::ToPugiEncoding(InEncoding);
+	Impl->XmlDoc.save(*Impl->XmlWriter, "  ", mxpugi::format_default, encoding);
 }
 
 void FXmlStructSerializerBackend::BeginArray(const FStructSerializerState& InState)
 {
-	pugi::xml_node ArrayElement = Impl->XmlElement.append_child(pugi::node_element);
+	mxpugi::xml_node ArrayElement = Impl->XmlElement.append_child(mxpugi::node_element);
 	ArrayElement.set_name("Array");
 	if (InState.ValueProperty)
 	{
@@ -68,7 +68,7 @@ void FXmlStructSerializerBackend::BeginArray(const FStructSerializerState& InSta
 
 void FXmlStructSerializerBackend::BeginStructure(const FStructSerializerState& InState)
 {
-	pugi::xml_node StructElement = Impl->XmlElement.append_child(pugi::node_element);
+	mxpugi::xml_node StructElement = Impl->XmlElement.append_child(mxpugi::node_element);
 	StructElement.set_name("Struct");
 	if (InState.ValueProperty)
 	{
@@ -98,14 +98,14 @@ void FXmlStructSerializerBackend::EndStructure(const FStructSerializerState& /*S
 
 void FXmlStructSerializerBackend::WriteComment(const FString& InComment)
 {
-	pugi::xml_node Child = Impl->XmlElement.append_child(pugi::node_comment);
+	mxpugi::xml_node Child = Impl->XmlElement.append_child(mxpugi::node_comment);
 	Child.set_value(STRING_CAST_UE_TO_PUGI(*InComment));
 }
 
 void FXmlStructSerializerBackend::WriteProperty(const FStructSerializerState& InState, int32 InArrayIndex)
 {
 	// Append a property element in the current structure element.
-	pugi::xml_node PropertyElement = Impl->XmlElement.append_child(pugi::node_element);
+	mxpugi::xml_node PropertyElement = Impl->XmlElement.append_child(mxpugi::node_element);
 	PropertyElement.set_name("Property");
 	PropertyElement.append_attribute("Name") = STRING_CAST_UE_TO_PUGI(*InState.ValueProperty->GetName());
 	PropertyElement.append_attribute("Type") = STRING_CAST_UE_TO_PUGI(*InState.ValueProperty->GetClass()->GetName());
@@ -113,7 +113,7 @@ void FXmlStructSerializerBackend::WriteProperty(const FStructSerializerState& In
 	{
 		PropertyElement.append_attribute("ValueType") = STRING_CAST_UE_TO_PUGI(*InState.ValueType->GetName());
 	}
-	pugi::xml_attribute ValueAttribute = PropertyElement.append_attribute("Value");
+	mxpugi::xml_attribute ValueAttribute = PropertyElement.append_attribute("Value");
 
 	// booleans
 	if (InState.FieldType == FBoolProperty::StaticClass())
diff --git a/Engine/Plugins/Experimental/SerializationUtils/Source/XmlSerialization/Private/Formatters/XmlArchiveInputFormatter.cpp b/Engine/Plugins/Experimental/SerializationUtils/Source/XmlSerialization/Private/Formatters/XmlArchiveInputFormatter.cpp
--- a/Engine/Plugins/Experimental/SerializationUtils/Source/XmlSerialization/Private/Formatters/XmlArchiveInputFormatter.cpp
+++ b/Engine/Plugins/Experimental/SerializationUtils/Source/XmlSerialization/Private/Formatters/XmlArchiveInputFormatter.cpp
@@ -52,15 +52,15 @@ public:
 		return true;
 	}
 
-	void Push(const pugi::xml_node& InXmlNode)
+	void Push(const mxpugi::xml_node& InXmlNode)
 	{
 		XmlElements.Push(InXmlNode);
 	}
 
-	pugi::xml_attribute GetValueAttrib() const
+	mxpugi::xml_attribute GetValueAttrib() const
 	{
 		using namespace UE::XmlSerialization::Private;
-		const pugi::xml_attribute Attrib = !XmlAttrib.empty() ? XmlAttrib : XmlElement().attribute("Value");
+		const mxpugi::xml_attribute Attrib = !XmlAttrib.empty() ? XmlAttrib : XmlElement().attribute("Value");
 		if (Attrib.empty())
 		{
 			UE_LOG(LogXmlSerialization, Error,
@@ -71,28 +71,28 @@ public:
 	}
 
 	/** Current xml element being traversed. */
-	const pugi::xml_node& XmlElement() const { return XmlElements.Top(); }
+	const mxpugi::xml_node& XmlElement() const { return XmlElements.Top(); }
 
-	pugi::xml_node& XmlElementMutable() { return XmlElements.Top(); }
+	mxpugi::xml_node& XmlElementMutable() { return XmlElements.Top(); }
 	
 	/** Current Document. */
-	pugi::xml_document XmlDocument;
+	mxpugi::xml_document XmlDocument;
 
 	/** Parse result of parsing the document. */
-	pugi::xml_parse_result XmlParseResult;
+	mxpugi::xml_parse_result XmlParseResult;
 
 	/**
 	 * Stack of xml element being traversed.
 	 * Using a stack so we can push empty nodes and pop them back up to the previous existing parent.
 	 * This allows traversing missing tags and recover.
 	 */
-	TArray<pugi::xml_node> XmlElements;
+	TArray<mxpugi::xml_node> XmlElements;
 
 	/** Current Attribute being traversed. */
-	pugi::xml_attribute XmlAttrib;
+	mxpugi::xml_attribute XmlAttrib;
 
 	/** Last Array/Stream/Map element that was traversed. */
-	pugi::xml_node LastArrayElement;
+	mxpugi::xml_node LastArrayElement;
 };
 
 FXmlArchiveInputFormatter::FXmlArchiveInputFormatter(FArchive& InInner, UObject* InRootObject, const TFunction<UObject* (const FPackageIndex)>& InResolveObject) 
@@ -110,9 +110,9 @@ FXmlArchiveInputFormatter::FXmlArchiveInputFormatter(FArchive& InInner, UObject*
 		Buffer.SetNum(Inner.TotalSize());
 		Inner.Serialize(Buffer.GetData(), Buffer.Num());
 		Impl->XmlParseResult = Impl->XmlDocument.load_buffer(Buffer.GetData(), Buffer.Num(),
-		                                                              pugi::parse_default | pugi::parse_declaration |
-		                                                              pugi::parse_comments);
-		if (Impl->XmlParseResult.status != pugi::status_ok)
+		                                                              mxpugi::parse_default | mxpugi::parse_declaration |
+		                                                              mxpugi::parse_comments);
+		if (Impl->XmlParseResult.status != mxpugi::status_ok)
 		{
 			//Impl->XmlParseResult.description() -> convert to FString.
 			UE_LOG(LogXmlSerialization, Error, TEXT("FXmlArchiveInputFormatter: Failed to parse xml document.")); 
@@ -120,7 +120,7 @@ FXmlArchiveInputFormatter::FXmlArchiveInputFormatter(FArchive& InInner, UObject*
 		else
 		{
 			// Find the root of the document.
-			pugi::xml_node RootNode = Impl->XmlDocument.document_element();
+			mxpugi::xml_node RootNode = Impl->XmlDocument.document_element();
 			if (FCStringAnsi::Strcmp(RootNode.name(), "ObjectFileDocument") != 0)
 			{
 				RootNode = RootNode.next_sibling("ObjectFileDocument");
@@ -155,7 +155,7 @@ FXmlArchiveInputFormatter::~FXmlArchiveInputFormatter()
 
 bool FXmlArchiveInputFormatter::IsParseResultStatusOk() const
 {
-	return (Impl->XmlParseResult.status == pugi::status_ok);
+	return (Impl->XmlParseResult.status == mxpugi::status_ok);
 }
 
 FArchive& FXmlArchiveInputFormatter::GetUnderlyingArchive()
@@ -180,7 +180,7 @@ bool FXmlArchiveInputFormatter::HasDocumentTree() const
 void FXmlArchiveInputFormatter::EnterRecord()
 {
 	using namespace UE::XmlSerialization::Private;
-	const pugi::xml_node RecordElement = Impl->XmlElement().child("Record");
+	const mxpugi::xml_node RecordElement = Impl->XmlElement().child("Record");
 	if (RecordElement.empty())
 	{
 		UE_LOG(LogXmlSerialization, Warning,
@@ -201,7 +201,7 @@ void FXmlArchiveInputFormatter::EnterField(FArchiveFieldName InName)
 	using namespace UE::XmlSerialization::Private;
 	// We are looking for a "Field" node with the attribute "Name" being the given name.
 	// Note: this only searches by attribute name/value.
-	const pugi::xml_node FieldElement = Impl->XmlElement().find_child_by_attribute("Field", "Name", STRING_CAST_UE_TO_PUGI(InName.Name));
+	const mxpugi::xml_node FieldElement = Impl->XmlElement().find_child_by_attribute("Field", "Name", STRING_CAST_UE_TO_PUGI(InName.Name));
 	if (FieldElement.empty())
 	{
 		UE_LOG(LogXmlSerialization, Warning,
@@ -219,7 +219,7 @@ void FXmlArchiveInputFormatter::LeaveField()
 
 bool FXmlArchiveInputFormatter::TryEnterField(FArchiveFieldName InName, bool bInEnterWhenSaving)
 {
-	const pugi::xml_node FieldElement = Impl->XmlElement().find_child_by_attribute("Field", "Name", STRING_CAST_UE_TO_PUGI(InName.Name));
+	const mxpugi::xml_node FieldElement = Impl->XmlElement().find_child_by_attribute("Field", "Name", STRING_CAST_UE_TO_PUGI(InName.Name));
 	if (FieldElement.empty())
 	{
 		return false; 
@@ -232,11 +232,11 @@ bool FXmlArchiveInputFormatter::TryEnterField(FArchiveFieldName InName, bool bIn
 void FXmlArchiveInputFormatter::EnterArray(int32& OutNumElements)
 {
 	using namespace UE::XmlSerialization::Private;
-	const pugi::xml_node ArrayElement = Impl->XmlElement().child("Array");
+	const mxpugi::xml_node ArrayElement = Impl->XmlElement().child("Array");
 	
 	if (!ArrayElement.empty())
 	{
-		const pugi::xml_attribute SizeAttrib = ArrayElement.attribute("Size");
+		const mxpugi::xml_attribute SizeAttrib = ArrayElement.attribute("Size");
 		if (SizeAttrib.empty())
 		{
 			UE_LOG(LogXmlSerialization, Error,
@@ -255,20 +255,20 @@ void FXmlArchiveInputFormatter::EnterArray(int32& OutNumElements)
 			*Utils::GetInfo(Impl->XmlElement())); 
 	}
 	
-	Impl->LastArrayElement = pugi::xml_node();
+	Impl->LastArrayElement = mxpugi::xml_node();
 	Impl->Push(ArrayElement);
 }
 
 void FXmlArchiveInputFormatter::LeaveArray()
 {
 	Impl->Pop(TEXT("LeaveArray"), TEXT("Array"));
-	Impl->LastArrayElement = pugi::xml_node();
+	Impl->LastArrayElement = mxpugi::xml_node();
 }
 
 void FXmlArchiveInputFormatter::EnterArrayElement()
 {
 	using namespace UE::XmlSerialization::Private;
-	pugi::xml_node ArrayElementElement;
+	mxpugi::xml_node ArrayElementElement;
 	
 	if (!Impl->LastArrayElement.empty())
 	{
@@ -299,14 +299,14 @@ void FXmlArchiveInputFormatter::EnterArrayElement()
 void FXmlArchiveInputFormatter::LeaveArrayElement()
 {
 	using namespace UE::XmlSerialization::Private;
-	Impl->LastArrayElement =  Utils::IsElementA(Impl->XmlElement(),TEXT("ArrayElement")) ? Impl->XmlElement() : pugi::xml_node();
+	Impl->LastArrayElement =  Utils::IsElementA(Impl->XmlElement(),TEXT("ArrayElement")) ? Impl->XmlElement() : mxpugi::xml_node();
 	Impl->Pop(TEXT("LeaveArrayElement"), TEXT("ArrayElement"));
 }
 
 void FXmlArchiveInputFormatter::EnterStream()
 {
 	using namespace UE::XmlSerialization::Private;
-	const pugi::xml_node StreamElement = Impl->XmlElement().child("Stream");
+	const mxpugi::xml_node StreamElement = Impl->XmlElement().child("Stream");
 
 	if (StreamElement.empty())
 	{
@@ -316,19 +316,19 @@ void FXmlArchiveInputFormatter::EnterStream()
 	}
 	
 	Impl->Push(StreamElement);
-	Impl->LastArrayElement = pugi::xml_node();
+	Impl->LastArrayElement = mxpugi::xml_node();
 }
 
 void FXmlArchiveInputFormatter::LeaveStream()
 {
 	Impl->Pop(TEXT("LeaveStream"), TEXT("Stream"));
-	Impl->LastArrayElement = pugi::xml_node();
+	Impl->LastArrayElement = mxpugi::xml_node();
 }
 
 void FXmlArchiveInputFormatter::EnterStreamElement()
 {
 	using namespace UE::XmlSerialization::Private;
-	pugi::xml_node StreamElementElement;
+	mxpugi::xml_node StreamElementElement;
 	
 	if (!Impl->LastArrayElement.empty())
 	{
@@ -359,17 +359,17 @@ void FXmlArchiveInputFormatter::EnterStreamElement()
 void FXmlArchiveInputFormatter::LeaveStreamElement()
 {
 	using namespace UE::XmlSerialization::Private;
-	Impl->LastArrayElement =  Utils::IsElementA(Impl->XmlElement(),TEXT("StreamElement")) ? Impl->XmlElement() : pugi::xml_node();
+	Impl->LastArrayElement =  Utils::IsElementA(Impl->XmlElement(),TEXT("StreamElement")) ? Impl->XmlElement() : mxpugi::xml_node();
 	Impl->Pop(TEXT("LeaveStreamElement"), TEXT("StreamElement"));
 }
 
 void FXmlArchiveInputFormatter::EnterMap(int32& OutNumElements)
 {
 	using namespace UE::XmlSerialization::Private;
-	const pugi::xml_node MapElement = Impl->XmlElement().child("Map");
+	const mxpugi::xml_node MapElement = Impl->XmlElement().child("Map");
 	if (!MapElement.empty())
 	{
-		const pugi::xml_attribute SizeAttrib = MapElement.attribute("Size");
+		const mxpugi::xml_attribute SizeAttrib = MapElement.attribute("Size");
 
 		if (SizeAttrib.empty())
 		{
@@ -379,13 +379,13 @@ void FXmlArchiveInputFormatter::EnterMap(int32& OutNumElements)
 		OutNumElements = !SizeAttrib.empty() ? SizeAttrib.as_int() : 0;
 		
 		Impl->Push(MapElement);
-		Impl->LastArrayElement = pugi::xml_node();
+		Impl->LastArrayElement = mxpugi::xml_node();
 		return;
 	}
 	
 	// Fallback mode: There might be a bug on how MapProperties (or Struct?) are serialized. The code is different for loading vs saving.
 	// See: UStruct::SerializeVersionedTaggedProperties. When loading it wants to enter a map, while when saving, it enters a record.
-	pugi::xml_node RecordElement = Impl->XmlElement().child("Record");
+	mxpugi::xml_node RecordElement = Impl->XmlElement().child("Record");
 	if (!RecordElement.empty())
 	{
 		// Remember we where supposed to get into a Map.
@@ -394,14 +394,14 @@ void FXmlArchiveInputFormatter::EnterMap(int32& OutNumElements)
 		// The number of elements are going to be the number of Field nodes.
 		// It will then try to enter MapElements, but we'll be looking for Field Nodes instead.
 		OutNumElements = 0;
-		for (pugi::xml_node Field : RecordElement.children("Field"))
+		for (mxpugi::xml_node Field : RecordElement.children("Field"))
 		{
 			OutNumElements++;
 		}	
 
 		UE_LOG(LogXmlSerialization, Verbose, TEXT("FXmlArchiveInputFormatter:EnterMap: Fallback: Entering \"Record\" node instead, with %d entries (Fields)."), OutNumElements);
 		Impl->Push(RecordElement);
-		Impl->LastArrayElement = pugi::xml_node();
+		Impl->LastArrayElement = mxpugi::xml_node();
 		return;
 	}
 
@@ -410,15 +410,15 @@ void FXmlArchiveInputFormatter::EnterMap(int32& OutNumElements)
 		*Utils::GetInfo(Impl->XmlElement()));
 	
 	OutNumElements = 0;
-	Impl->Push(pugi::xml_node());
-	Impl->LastArrayElement = pugi::xml_node();
+	Impl->Push(mxpugi::xml_node());
+	Impl->LastArrayElement = mxpugi::xml_node();
 }
 
-inline bool ElementIsARecordButShouldBeAMap(const pugi::xml_node& InElement)
+inline bool ElementIsARecordButShouldBeAMap(const mxpugi::xml_node& InElement)
 {
 	if (FCStringAnsi::Strcmp(InElement.name(), "Record") == 0)
 	{
-		const pugi::xml_attribute HackAttrib = InElement.attribute("__HACK_SHOULD_BE__");
+		const mxpugi::xml_attribute HackAttrib = InElement.attribute("__HACK_SHOULD_BE__");
 		return (!HackAttrib.empty() && FCStringAnsi::Strcmp(HackAttrib.as_string(), "Map") == 0) ? true : false;
 	}
 	return false;
@@ -434,7 +434,7 @@ void FXmlArchiveInputFormatter::LeaveMap()
 	{
 		Impl->Pop(TEXT("LeaveMap"), TEXT("Map"));
 	}
-	Impl->LastArrayElement = pugi::xml_node();
+	Impl->LastArrayElement = mxpugi::xml_node();
 }
 
 void FXmlArchiveInputFormatter::EnterMapElement(FString& OutName)
@@ -444,7 +444,7 @@ void FXmlArchiveInputFormatter::EnterMapElement(FString& OutName)
 	// Fallback code for the Map vs Record bug.
 	if (ElementIsARecordButShouldBeAMap(Impl->XmlElement()))
 	{
-		pugi::xml_node FieldElement;
+		mxpugi::xml_node FieldElement;
 		
 		if (!Impl->LastArrayElement.empty())
 		{
@@ -471,7 +471,7 @@ void FXmlArchiveInputFormatter::EnterMapElement(FString& OutName)
 
 		if (!FieldElement.empty())
 		{
-			const pugi::xml_attribute NameAttrib = FieldElement.attribute("Name");
+			const mxpugi::xml_attribute NameAttrib = FieldElement.attribute("Name");
 			OutName = NameAttrib.as_string();
 			if (NameAttrib.empty())
 			{
@@ -488,7 +488,7 @@ void FXmlArchiveInputFormatter::EnterMapElement(FString& OutName)
 		return; // fallback job done.
 	}
 	
-	pugi::xml_node MapElementElement;
+	mxpugi::xml_node MapElementElement;
 	
 	if (!Impl->LastArrayElement.empty())
 	{
@@ -515,7 +515,7 @@ void FXmlArchiveInputFormatter::EnterMapElement(FString& OutName)
 	
 	if (!MapElementElement.empty())
 	{
-		const pugi::xml_attribute NameAttrib = MapElementElement.attribute("Name");
+		const mxpugi::xml_attribute NameAttrib = MapElementElement.attribute("Name");
 		OutName = NameAttrib.as_string();
 		if (NameAttrib.empty())
 		{
@@ -533,7 +533,7 @@ void FXmlArchiveInputFormatter::LeaveMapElement()
 	using namespace UE::XmlSerialization::Private;
 	
 	// Check if our current element should be something else (fallback code above).
-	const pugi::xml_attribute HackAttrib = Impl->XmlElement().attribute("__HACK_SHOULD_BE__");
+	const mxpugi::xml_attribute HackAttrib = Impl->XmlElement().attribute("__HACK_SHOULD_BE__");
 	if (!HackAttrib.empty())
 	{
 		// It is possible to be in a value node. If so pop it first.
@@ -543,12 +543,12 @@ void FXmlArchiveInputFormatter::LeaveMapElement()
 		}
 
 		// Should be in a Field node. If so, keep track of the last element.
-		Impl->LastArrayElement =  Utils::IsElementA(Impl->XmlElement(),TEXT("Field")) ? Impl->XmlElement() : pugi::xml_node();
+		Impl->LastArrayElement =  Utils::IsElementA(Impl->XmlElement(),TEXT("Field")) ? Impl->XmlElement() : mxpugi::xml_node();
 		Impl->Pop(TEXT("LeaveMapElement"), TEXT("Field"));
 	}
 	else
 	{
-		Impl->LastArrayElement =  Utils::IsElementA(Impl->XmlElement(),TEXT("MapElement")) ? Impl->XmlElement() : pugi::xml_node();
+		Impl->LastArrayElement =  Utils::IsElementA(Impl->XmlElement(),TEXT("MapElement")) ? Impl->XmlElement() : mxpugi::xml_node();
 		Impl->Pop(TEXT("LeaveMapElement"), TEXT("MapElement"));
 	}
 }
@@ -556,7 +556,7 @@ void FXmlArchiveInputFormatter::LeaveMapElement()
 void FXmlArchiveInputFormatter::EnterAttributedValue()
 {
 	using namespace UE::XmlSerialization::Private;
-	const pugi::xml_node ValueElement = Impl->XmlElement().child("AttributedValue");
+	const mxpugi::xml_node ValueElement = Impl->XmlElement().child("AttributedValue");
 	if (ValueElement.empty())
 	{
 		UE_LOG(LogXmlSerialization, Error,
@@ -569,7 +569,7 @@ void FXmlArchiveInputFormatter::EnterAttributedValue()
 void FXmlArchiveInputFormatter::EnterAttribute(FArchiveFieldName AttributeName)
 {
 	using namespace UE::XmlSerialization::Private;
-	const pugi::xml_node AttributeElement = Impl->XmlElement().find_child_by_attribute("Attribute", "Name", STRING_CAST_UE_TO_PUGI(AttributeName.Name));
+	const mxpugi::xml_node AttributeElement = Impl->XmlElement().find_child_by_attribute("Attribute", "Name", STRING_CAST_UE_TO_PUGI(AttributeName.Name));
 	if (AttributeElement.empty())
 	{
 		UE_LOG(LogXmlSerialization, Error,
@@ -583,7 +583,7 @@ void FXmlArchiveInputFormatter::EnterAttribute(FArchiveFieldName AttributeName)
 void FXmlArchiveInputFormatter::EnterAttributedValueValue()
 {
 	using namespace UE::XmlSerialization::Private;
-	const pugi::xml_node ValueElement = Impl->XmlElement().child("AttributedValueValue");
+	const mxpugi::xml_node ValueElement = Impl->XmlElement().child("AttributedValueValue");
 	if (ValueElement.empty())
 	{
 		UE_LOG(LogXmlSerialization, Error,
@@ -595,7 +595,7 @@ void FXmlArchiveInputFormatter::EnterAttributedValueValue()
 
 bool FXmlArchiveInputFormatter::TryEnterAttributedValueValue()
 {
-	const pugi::xml_node ValueElement = Impl->XmlElement().child("AttributedValueValue");
+	const mxpugi::xml_node ValueElement = Impl->XmlElement().child("AttributedValueValue");
 	if (!ValueElement.empty())
 	{
 		Impl->Push(ValueElement);
@@ -638,7 +638,7 @@ void FXmlArchiveInputFormatter::LeaveAttributedValue()
 
 bool FXmlArchiveInputFormatter::TryEnterAttribute(FArchiveFieldName InAttributeName, bool bInEnterWhenSaving)
 {
-	const pugi::xml_node AttributeElement = Impl->XmlElement().find_child_by_attribute("Attribute", "Name", STRING_CAST_UE_TO_PUGI(InAttributeName.Name));
+	const mxpugi::xml_node AttributeElement = Impl->XmlElement().find_child_by_attribute("Attribute", "Name", STRING_CAST_UE_TO_PUGI(InAttributeName.Name));
 	if (!AttributeElement.empty())
 	{
 		Impl->Push(AttributeElement);
@@ -649,7 +649,7 @@ bool FXmlArchiveInputFormatter::TryEnterAttribute(FArchiveFieldName InAttributeN
 
 void FXmlArchiveInputFormatter::Serialize(uint8& Value)
 {
-	const pugi::xml_attribute Attrib = Impl->GetValueAttrib();
+	const mxpugi::xml_attribute Attrib = Impl->GetValueAttrib();
 	if (!Attrib.empty())
 	{
 		Value = Attrib.as_uint();
@@ -658,7 +658,7 @@ void FXmlArchiveInputFormatter::Serialize(uint8& Value)
 
 void FXmlArchiveInputFormatter::Serialize(uint16& Value)
 {
-	const pugi::xml_attribute Attrib = Impl->GetValueAttrib();
+	const mxpugi::xml_attribute Attrib = Impl->GetValueAttrib();
 	if (!Attrib.empty())
 	{
 		Value = Attrib.as_uint();
@@ -667,7 +667,7 @@ void FXmlArchiveInputFormatter::Serialize(uint16& Value)
 
 void FXmlArchiveInputFormatter::Serialize(uint32& Value)
 {
-	const pugi::xml_attribute Attrib = Impl->GetValueAttrib();
+	const mxpugi::xml_attribute Attrib = Impl->GetValueAttrib();
 	if (!Attrib.empty())
 	{
 		Value = Attrib.as_uint();
@@ -676,7 +676,7 @@ void FXmlArchiveInputFormatter::Serialize(uint32& Value)
 
 void FXmlArchiveInputFormatter::Serialize(uint64& Value)
 {
-	const pugi::xml_attribute Attrib = Impl->GetValueAttrib();
+	const mxpugi::xml_attribute Attrib = Impl->GetValueAttrib();
 	if (!Attrib.empty())
 	{
 		Value = Attrib.as_ullong();
@@ -685,7 +685,7 @@ void FXmlArchiveInputFormatter::Serialize(uint64& Value)
 
 void FXmlArchiveInputFormatter::Serialize(int8& Value)
 {
-	const pugi::xml_attribute Attrib = Impl->GetValueAttrib();
+	const mxpugi::xml_attribute Attrib = Impl->GetValueAttrib();
 	if (!Attrib.empty())
 	{
 		Value = Attrib.as_int();
@@ -694,7 +694,7 @@ void FXmlArchiveInputFormatter::Serialize(int8& Value)
 
 void FXmlArchiveInputFormatter::Serialize(int16& Value)
 {
-	const pugi::xml_attribute Attrib = Impl->GetValueAttrib();
+	const mxpugi::xml_attribute Attrib = Impl->GetValueAttrib();
 	if (!Attrib.empty())
 	{
 		Value = Attrib.as_int();
@@ -703,7 +703,7 @@ void FXmlArchiveInputFormatter::Serialize(int16& Value)
 
 void FXmlArchiveInputFormatter::Serialize(int32& Value)
 {
-	const pugi::xml_attribute Attrib = Impl->GetValueAttrib();
+	const mxpugi::xml_attribute Attrib = Impl->GetValueAttrib();
 	if (!Attrib.empty())
 	{
 		Value = Attrib.as_int();
@@ -712,7 +712,7 @@ void FXmlArchiveInputFormatter::Serialize(int32& Value)
 
 void FXmlArchiveInputFormatter::Serialize(int64& Value)
 {
-	const pugi::xml_attribute Attrib = Impl->GetValueAttrib();
+	const mxpugi::xml_attribute Attrib = Impl->GetValueAttrib();
 	if (!Attrib.empty())
 	{
 		Value = Attrib.as_llong();
@@ -721,7 +721,7 @@ void FXmlArchiveInputFormatter::Serialize(int64& Value)
 
 void FXmlArchiveInputFormatter::Serialize(float& Value)
 {
-	const pugi::xml_attribute Attrib = Impl->GetValueAttrib();
+	const mxpugi::xml_attribute Attrib = Impl->GetValueAttrib();
 	if (!Attrib.empty())
 	{
 		Value = Attrib.as_float();
@@ -730,7 +730,7 @@ void FXmlArchiveInputFormatter::Serialize(float& Value)
 	
 void FXmlArchiveInputFormatter::Serialize(double& Value)
 {
-	const pugi::xml_attribute Attrib = Impl->GetValueAttrib();
+	const mxpugi::xml_attribute Attrib = Impl->GetValueAttrib();
 	if (!Attrib.empty())
 	{
 		Value = Attrib.as_double();
@@ -739,7 +739,7 @@ void FXmlArchiveInputFormatter::Serialize(double& Value)
 
 void FXmlArchiveInputFormatter::Serialize(bool& Value)
 {
-	const pugi::xml_attribute Attrib = Impl->GetValueAttrib();
+	const mxpugi::xml_attribute Attrib = Impl->GetValueAttrib();
 	if (!Attrib.empty())
 	{
 		Value = Attrib.as_bool();
@@ -748,7 +748,7 @@ void FXmlArchiveInputFormatter::Serialize(bool& Value)
 
 void FXmlArchiveInputFormatter::Serialize(FString& Value)
 {
-	const pugi::xml_attribute Attrib = Impl->GetValueAttrib();
+	const mxpugi::xml_attribute Attrib = Impl->GetValueAttrib();
 	if (!Attrib.empty())
 	{
 		Value = Attrib.as_string();	
@@ -758,7 +758,7 @@ void FXmlArchiveInputFormatter::Serialize(FString& Value)
 
 void FXmlArchiveInputFormatter::Serialize(FName& Value)
 {
-	const pugi::xml_attribute Attrib = Impl->GetValueAttrib();
+	const mxpugi::xml_attribute Attrib = Impl->GetValueAttrib();
 	if (!Attrib.empty())
 	{
 		const FString StringValue = Attrib.as_string();
@@ -775,7 +775,7 @@ void FXmlArchiveInputFormatter::Serialize(UObject*& Value)
 	}
 	else
 	{
-		const pugi::xml_attribute Attrib = Impl->GetValueAttrib();
+		const mxpugi::xml_attribute Attrib = Impl->GetValueAttrib();
 		const FString StringValue = !Attrib.empty() ? Attrib.as_string() : FString();
 		if (!StringValue.IsEmpty())
 		{
@@ -881,7 +881,7 @@ void FXmlArchiveInputFormatter::Serialize(void* Data, uint64 DataSize)
 
 bool FXmlArchiveInputFormatter::IsNestedObject() const
 {
-	const pugi::xml_node RecordChild = Impl->XmlElement().child("Record");
+	const mxpugi::xml_node RecordChild = Impl->XmlElement().child("Record");
 	if (!RecordChild.empty() && RecordChild.attribute("Class").empty() == false)
 	{
 		return true;
@@ -892,7 +892,7 @@ bool FXmlArchiveInputFormatter::IsNestedObject() const
 UObject* FXmlArchiveInputFormatter::LoadNestedObject()
 {
 	using namespace UE::XmlSerialization::Private;
-	const pugi::xml_node RecordChild = Impl->XmlElement().child("Record");
+	const mxpugi::xml_node RecordChild = Impl->XmlElement().child("Record");
 	if (!RecordChild.empty() && RecordChild.attribute("Class").empty() == false)
 	{
 		const FSoftObjectPath ClassPath(RecordChild.attribute("Class").as_string());
@@ -912,7 +912,7 @@ UObject* FXmlArchiveInputFormatter::LoadNestedObject()
 
 			{
 				// Rename the root field tag's name attribute to match what the new object's name is.
-				const pugi::xml_node FieldElement = RecordChild.find_child_by_attribute("Field", "Name", STRING_CAST_UE_TO_PUGI(*OriginalObjectName));
+				const mxpugi::xml_node FieldElement = RecordChild.find_child_by_attribute("Field", "Name", STRING_CAST_UE_TO_PUGI(*OriginalObjectName));
 				if (!FieldElement.empty())
 				{
 					FieldElement.attribute("Name") = STRING_CAST_UE_TO_PUGI(*NestedObject->GetName());
diff --git a/Engine/Plugins/Experimental/SerializationUtils/Source/XmlSerialization/Private/Formatters/XmlArchiveOutputFormatter.cpp b/Engine/Plugins/Experimental/SerializationUtils/Source/XmlSerialization/Private/Formatters/XmlArchiveOutputFormatter.cpp
--- a/Engine/Plugins/Experimental/SerializationUtils/Source/XmlSerialization/Private/Formatters/XmlArchiveOutputFormatter.cpp
+++ b/Engine/Plugins/Experimental/SerializationUtils/Source/XmlSerialization/Private/Formatters/XmlArchiveOutputFormatter.cpp
@@ -23,7 +23,7 @@ class FXmlArchiveOutputFormatter::FImpl
 public:
 	FImpl()
 	{
-		pugi::xml_node decl = XmlDoc.prepend_child(pugi::node_declaration);
+		mxpugi::xml_node decl = XmlDoc.prepend_child(mxpugi::node_declaration);
 		decl.append_attribute("version") = "1.0";
 		decl.append_attribute("encoding") = "utf-8";
 
@@ -82,9 +82,9 @@ public:
 	}
 
 	
-	pugi::xml_document XmlDoc;
-	pugi::xml_node XmlElement;
-	pugi::xml_attribute XmlAttrib;
+	mxpugi::xml_document XmlDoc;
+	mxpugi::xml_node XmlElement;
+	mxpugi::xml_attribute XmlAttrib;
 	bool bDocumentSaveAttempted = false;
 };
 
@@ -121,10 +121,10 @@ bool FXmlArchiveOutputFormatter::SaveDocumentToInnerArchive(EXmlSerializationEnc
 	bool bSuccess = true;
 	try
 	{
-		const pugi::xml_encoding encoding = Utils::ToPugiEncoding(InEncoding);
+		const mxpugi::xml_encoding encoding = Utils::ToPugiEncoding(InEncoding);
 		Impl->bDocumentSaveAttempted = true;
 		FXmlWriter XmlWriter(&Inner);
-		Impl->XmlDoc.save(XmlWriter, "  ", pugi::format_default, encoding);
+		Impl->XmlDoc.save(XmlWriter, "  ", mxpugi::format_default, encoding);
 	}
 	catch (...)
 	{
@@ -146,7 +146,7 @@ bool FXmlArchiveOutputFormatter::HasDocumentTree() const
 
 void FXmlArchiveOutputFormatter::EnterRecord()
 {
-	pugi::xml_node RecordElement = Impl->XmlElement.append_child(pugi::node_element);
+	mxpugi::xml_node RecordElement = Impl->XmlElement.append_child(mxpugi::node_element);
 	RecordElement.set_name("Record");
 	Impl->XmlElement = RecordElement;	
 }
@@ -158,7 +158,7 @@ void FXmlArchiveOutputFormatter::LeaveRecord()
 
 void FXmlArchiveOutputFormatter::EnterField(FArchiveFieldName Name)
 {
-	pugi::xml_node FieldElement = Impl->XmlElement.append_child(pugi::node_element);
+	mxpugi::xml_node FieldElement = Impl->XmlElement.append_child(mxpugi::node_element);
 	FieldElement.set_name("Field");
 	FieldElement.append_attribute("Name") = STRING_CAST_UE_TO_PUGI(Name.Name);
 	Impl->XmlElement = FieldElement;
@@ -180,7 +180,7 @@ bool FXmlArchiveOutputFormatter::TryEnterField(FArchiveFieldName Name, bool bEnt
 
 void FXmlArchiveOutputFormatter::EnterArray(int32& NumElements)
 {
-	pugi::xml_node ArrayElement = Impl->XmlElement.append_child(pugi::node_element);
+	mxpugi::xml_node ArrayElement = Impl->XmlElement.append_child(mxpugi::node_element);
 	ArrayElement.set_name("Array");
 	ArrayElement.append_attribute("Size").set_value(NumElements);
 	Impl->XmlElement = ArrayElement;
@@ -193,7 +193,7 @@ void FXmlArchiveOutputFormatter::LeaveArray()
 
 void FXmlArchiveOutputFormatter::EnterArrayElement()
 {
-	pugi::xml_node ArrayElementElement = Impl->XmlElement.append_child(pugi::node_element);
+	mxpugi::xml_node ArrayElementElement = Impl->XmlElement.append_child(mxpugi::node_element);
 	ArrayElementElement.set_name("ArrayElement");
 	Impl->XmlElement = ArrayElementElement;
 }
@@ -205,7 +205,7 @@ void FXmlArchiveOutputFormatter::LeaveArrayElement()
 
 void FXmlArchiveOutputFormatter::EnterStream()
 {
-	pugi::xml_node StreamElement = Impl->XmlElement.append_child(pugi::node_element);
+	mxpugi::xml_node StreamElement = Impl->XmlElement.append_child(mxpugi::node_element);
 	StreamElement.set_name("Stream");
 	Impl->XmlElement = StreamElement;
 }
@@ -217,7 +217,7 @@ void FXmlArchiveOutputFormatter::LeaveStream()
 
 void FXmlArchiveOutputFormatter::EnterStreamElement()
 {
-	pugi::xml_node StreamElementElement = Impl->XmlElement.append_child(pugi::node_element);
+	mxpugi::xml_node StreamElementElement = Impl->XmlElement.append_child(mxpugi::node_element);
 	StreamElementElement.set_name("StreamElement");
 	Impl->XmlElement = StreamElementElement;
 }
@@ -229,7 +229,7 @@ void FXmlArchiveOutputFormatter::LeaveStreamElement()
 
 void FXmlArchiveOutputFormatter::EnterMap(int32& NumElements)
 {
-	pugi::xml_node MapElement = Impl->XmlElement.append_child(pugi::node_element);
+	mxpugi::xml_node MapElement = Impl->XmlElement.append_child(mxpugi::node_element);
 	MapElement.set_name("Map");
 	MapElement.append_attribute("Size").set_value(NumElements);
 	Impl->XmlElement = MapElement;
@@ -242,7 +242,7 @@ void FXmlArchiveOutputFormatter::LeaveMap()
 
 void FXmlArchiveOutputFormatter::EnterMapElement(FString& Name)
 {
-	pugi::xml_node MapElementElement = Impl->XmlElement.append_child(pugi::node_element);
+	mxpugi::xml_node MapElementElement = Impl->XmlElement.append_child(mxpugi::node_element);
 	MapElementElement.set_name("MapElement");
 	MapElementElement.append_attribute("Name") = STRING_CAST_UE_TO_PUGI(*Name);
 	Impl->XmlElement = MapElementElement;
@@ -255,14 +255,14 @@ void FXmlArchiveOutputFormatter::LeaveMapElement()
 
 void FXmlArchiveOutputFormatter::EnterAttributedValue()
 {
-	pugi::xml_node ValueElement = Impl->XmlElement.append_child(pugi::node_element);
+	mxpugi::xml_node ValueElement = Impl->XmlElement.append_child(mxpugi::node_element);
 	ValueElement.set_name("AttributedValue");
 	Impl->XmlElement = ValueElement;
 }
 
 void FXmlArchiveOutputFormatter::EnterAttribute(FArchiveFieldName AttributeName)
 {
-	pugi::xml_node AttributeElement = Impl->XmlElement.append_child(pugi::node_element);
+	mxpugi::xml_node AttributeElement = Impl->XmlElement.append_child(mxpugi::node_element);
 	AttributeElement.set_name("Attribute");
 	AttributeElement.append_attribute("Name") = STRING_CAST_UE_TO_PUGI(AttributeName.Name);
 	Impl->XmlElement = AttributeElement;
@@ -288,7 +288,7 @@ void FXmlArchiveOutputFormatter::LeaveAttribute()
 	{
 		Impl->Pop(TEXT("LeaveAttribute"), TEXT("Attribute"));
 	}
-	//Impl->XmlAttrib = pugi::xml_attribute();
+	//Impl->XmlAttrib = mxpugi::xml_attribute();
 }
 
 void FXmlArchiveOutputFormatter::LeaveAttributedValue()
@@ -303,7 +303,7 @@ void FXmlArchiveOutputFormatter::LeaveAttributedValue()
 
 void FXmlArchiveOutputFormatter::EnterAttributedValueValue()
 {
-	pugi::xml_node AttributedValueValueElement = Impl->XmlElement.append_child(pugi::node_element);
+	mxpugi::xml_node AttributedValueValueElement = Impl->XmlElement.append_child(mxpugi::node_element);
 	AttributedValueValueElement.set_name("AttributedValueValue");
 	Impl->XmlElement = AttributedValueValueElement;
 }
@@ -431,7 +431,7 @@ void FXmlArchiveOutputFormatter::Serialize(UObject*& Value)
 		if (Value != nullptr)
 		{
 			{
-				Impl->XmlElement.append_child(pugi::node_comment).set_value("Object In Place");
+				Impl->XmlElement.append_child(mxpugi::node_comment).set_value("Object In Place");
 			}
 			
 			FStructuredArchive ChildArchive(*this);
diff --git a/Engine/Plugins/Experimental/SerializationUtils/Source/XmlSerialization/Private/Utils/XmlUtils.h b/Engine/Plugins/Experimental/SerializationUtils/Source/XmlSerialization/Private/Utils/XmlUtils.h
--- a/Engine/Plugins/Experimental/SerializationUtils/Source/XmlSerialization/Private/Utils/XmlUtils.h
+++ b/Engine/Plugins/Experimental/SerializationUtils/Source/XmlSerialization/Private/Utils/XmlUtils.h
@@ -29,9 +29,9 @@
 namespace UE::XmlSerialization::Private::Utils
 {
 	// Helper for pretty log
-	inline FString GetInfo(const pugi::xml_node& XmlElement)
+	inline FString GetInfo(const mxpugi::xml_node& XmlElement)
 	{
-		const pugi::xml_attribute NameAttrib = XmlElement.attribute("Name");
+		const mxpugi::xml_attribute NameAttrib = XmlElement.attribute("Name");
 		if (!NameAttrib.empty())
 		{
 			return FString::Printf(TEXT("%s Name=\"%s\""), STRING_CAST_PUGI_TO_UE(XmlElement.name()), STRING_CAST_PUGI_TO_UE(NameAttrib.as_string()));
@@ -42,14 +42,14 @@ namespace UE::XmlSerialization::Private::Utils
 		}
 	}
 
-	inline bool IsElementA(const pugi::xml_node& Element, const TCHAR* TagName)
+	inline bool IsElementA(const mxpugi::xml_node& Element, const TCHAR* TagName)
 	{
 		return (FCStringAnsi::Strcmp(Element.name(), STRING_CAST_UE_TO_PUGI(TagName)) == 0) ? true : false;
 	}
 
-	inline bool IsElementShouldBe(const pugi::xml_node& Element, const TCHAR* TagName)
+	inline bool IsElementShouldBe(const mxpugi::xml_node& Element, const TCHAR* TagName)
 	{
-		const pugi::xml_attribute HackAttrib = Element.attribute("__HACK_SHOULD_BE__");
+		const mxpugi::xml_attribute HackAttrib = Element.attribute("__HACK_SHOULD_BE__");
 		if (!HackAttrib.empty())
 		{
 			return (FCStringAnsi::Strcmp(HackAttrib.as_string(), STRING_CAST_UE_TO_PUGI(TagName)) == 0) ? true : false;
@@ -57,28 +57,28 @@ namespace UE::XmlSerialization::Private::Utils
 		return false;
 	}
 	
-	inline pugi::xml_encoding ToPugiEncoding(EXmlSerializationEncoding InEncoding)
+	inline mxpugi::xml_encoding ToPugiEncoding(EXmlSerializationEncoding InEncoding)
 	{
 		switch (InEncoding)
 		{
 		case EXmlSerializationEncoding::Utf8:
-			return pugi::encoding_utf8;
+			return mxpugi::encoding_utf8;
 		case EXmlSerializationEncoding::Utf16_Le:
-			return pugi::encoding_utf16_le;
+			return mxpugi::encoding_utf16_le;
 		case EXmlSerializationEncoding::Utf16_Be:
-			return pugi::encoding_utf16_be;
+			return mxpugi::encoding_utf16_be;
 		case EXmlSerializationEncoding::Utf16:
-			return pugi::encoding_utf16;
+			return mxpugi::encoding_utf16;
 		case EXmlSerializationEncoding::Utf32_Le:
-			return pugi::encoding_utf32_le;
+			return mxpugi::encoding_utf32_le;
 		case EXmlSerializationEncoding::Utf32_Be:
-			return pugi::encoding_utf32_be;
+			return mxpugi::encoding_utf32_be;
 		case EXmlSerializationEncoding::Utf32:
-			return pugi::encoding_utf32;
+			return mxpugi::encoding_utf32;
 		case EXmlSerializationEncoding::WChar:
-			return pugi::encoding_wchar;
+			return mxpugi::encoding_wchar;
 		default:
-			return pugi::encoding_utf8;
+			return mxpugi::encoding_utf8;
 		}
 	}
 }
\ No newline at end of file
diff --git a/Engine/Plugins/Experimental/SerializationUtils/Source/XmlSerialization/Private/Utils/XmlWriter.h b/Engine/Plugins/Experimental/SerializationUtils/Source/XmlSerialization/Private/Utils/XmlWriter.h
--- a/Engine/Plugins/Experimental/SerializationUtils/Source/XmlSerialization/Private/Utils/XmlWriter.h
+++ b/Engine/Plugins/Experimental/SerializationUtils/Source/XmlSerialization/Private/Utils/XmlWriter.h
@@ -5,7 +5,7 @@
 #include "Serialization/Archive.h"
 #include "MaterialXFormat/External/PugiXML/pugixml.hpp"
 
-class FXmlWriter : public pugi::xml_writer
+class FXmlWriter : public mxpugi::xml_writer
 {
 public:
 	FXmlWriter(	FArchive* const InStream) : Stream(InStream) {}
