SerializationUtils, SVGImporter: use standalone installation of pugixml

diff --git a/Engine/Plugins/Experimental/SVGImporter/Source/SVGImporterEditor/Private/SVGParser_PugiXml.h b/Engine/Plugins/Experimental/SVGImporter/Source/SVGImporterEditor/Private/SVGParser_PugiXml.h
--- a/Engine/Plugins/Experimental/SVGImporter/Source/SVGImporterEditor/Private/SVGParser_PugiXml.h
+++ b/Engine/Plugins/Experimental/SVGImporter/Source/SVGImporterEditor/Private/SVGParser_PugiXml.h
@@ -3,7 +3,8 @@
 #pragma once
 
 #include "SVGParser_Base.h"
-#include "MaterialXFormat/External/PugiXML/pugixml.hpp"
+
+#include "pugixml.hpp"
 
 /**
  * SVG Parser based on pugi xml parser
diff --git a/Engine/Plugins/Experimental/SVGImporter/Source/SVGImporterEditor/SVGImporterEditor.Build.cs b/Engine/Plugins/Experimental/SVGImporter/Source/SVGImporterEditor/SVGImporterEditor.Build.cs
--- a/Engine/Plugins/Experimental/SVGImporter/Source/SVGImporterEditor/SVGImporterEditor.Build.cs
+++ b/Engine/Plugins/Experimental/SVGImporter/Source/SVGImporterEditor/SVGImporterEditor.Build.cs
@@ -31,9 +31,9 @@ public class SVGImporterEditor : ModuleRules
 				"Engine",
 				"InputCore",
 				"LevelEditor",
-				"MaterialX", // Used for pugi xml
 				"Projects",
 				"PropertyEditor",
+				"pugixml",
 				"RenderCore",
 				"RHI",
 				"Slate",
diff --git a/Engine/Plugins/Experimental/SerializationUtils/Source/XmlSerialization/Private/Backends/XmlStructSerializerBackend.cpp b/Engine/Plugins/Experimental/SerializationUtils/Source/XmlSerialization/Private/Backends/XmlStructSerializerBackend.cpp
--- a/Engine/Plugins/Experimental/SerializationUtils/Source/XmlSerialization/Private/Backends/XmlStructSerializerBackend.cpp
+++ b/Engine/Plugins/Experimental/SerializationUtils/Source/XmlSerialization/Private/Backends/XmlStructSerializerBackend.cpp
@@ -3,7 +3,6 @@
 #include "Backends/XmlStructSerializerBackend.h"
 
 #include "XmlSerializationModule.h"
-#include "MaterialXFormat/External/PugiXML/pugixml.hpp"
 #include "UObject/EnumProperty.h"
 #include "UObject/PropertyPortFlags.h"
 #include "UObject/TextProperty.h"
@@ -11,6 +10,8 @@
 #include "Utils/XmlUtils.h"
 #include "Utils/XmlWriter.h"
 
+#include "pugixml.hpp"
+
 class FXmlStructSerializerBackend::FImpl
 {
 public:
diff --git a/Engine/Plugins/Experimental/SerializationUtils/Source/XmlSerialization/Private/Utils/XmlUtils.h b/Engine/Plugins/Experimental/SerializationUtils/Source/XmlSerialization/Private/Utils/XmlUtils.h
--- a/Engine/Plugins/Experimental/SerializationUtils/Source/XmlSerialization/Private/Utils/XmlUtils.h
+++ b/Engine/Plugins/Experimental/SerializationUtils/Source/XmlSerialization/Private/Utils/XmlUtils.h
@@ -9,17 +9,15 @@
 #include "Misc/CString.h"
 #include "XmlSerializationDefines.h"
 
-#include "MaterialXFormat/External/PugiXML/pugixml.hpp"
+#include "pugixml.hpp"
 
 /**
- * Pugi can be compiled either in WChar mode (recommended) or Ansi mode.
- * Since we use the version shipped in MaterialXFormat (for now), we don't have a choice
- * to use that config, i.e. ainsi mode, but ideally we could have an independent
- * pugi third party lib in wchar mode. In the mean time, these macro are used to
- * make the code easy to change.
+ * pugixml can be compiled either in WChar mode (recommended) or Ansi mode.
+ * It should define PUGIXML_WCHAR_MODE in pugiconfig.hpp if it has been compiled
+ * in WChar mode.
  */
 #ifdef PUGIXML_WCHAR_MODE
-#	define	STRING_CAST_UE_TO_PUGI(InString) InString
+#	define STRING_CAST_UE_TO_PUGI(InString) InString
 #	define STRING_CAST_PUGI_TO_UE(InString) InString
 #else
 #	define STRING_CAST_UE_TO_PUGI(InString) StringCast<ANSICHAR>(InString).Get() 
diff --git a/Engine/Plugins/Experimental/SerializationUtils/Source/XmlSerialization/Private/Utils/XmlWriter.h b/Engine/Plugins/Experimental/SerializationUtils/Source/XmlSerialization/Private/Utils/XmlWriter.h
--- a/Engine/Plugins/Experimental/SerializationUtils/Source/XmlSerialization/Private/Utils/XmlWriter.h
+++ b/Engine/Plugins/Experimental/SerializationUtils/Source/XmlSerialization/Private/Utils/XmlWriter.h
@@ -3,7 +3,8 @@
 #pragma once
 
 #include "Serialization/Archive.h"
-#include "MaterialXFormat/External/PugiXML/pugixml.hpp"
+
+#include "pugixml.hpp"
 
 class FXmlWriter : public pugi::xml_writer
 {
diff --git a/Engine/Plugins/Experimental/SerializationUtils/Source/XmlSerialization/XmlSerialization.Build.cs b/Engine/Plugins/Experimental/SerializationUtils/Source/XmlSerialization/XmlSerialization.Build.cs
--- a/Engine/Plugins/Experimental/SerializationUtils/Source/XmlSerialization/XmlSerialization.Build.cs
+++ b/Engine/Plugins/Experimental/SerializationUtils/Source/XmlSerialization/XmlSerialization.Build.cs
@@ -17,12 +17,10 @@ public class XmlSerialization : ModuleRules
 			}
 		);
 
-		// PrivateDefinitions.Add("PUGIXML_WCHAR_MODE"); // for pugixml wchar support
-
 		PrivateDependencyModuleNames.AddRange(
 			new string[]
 			{
-				"MaterialX", // for pugixml
+				"pugixml",
 				"Serialization",
 			}
 		);
